WaveTable Checkpoint
- Migrate/Create base WaveTable functionality & editor bits from Modulation plugin so it can be leveraged for other assets & tools
- Add WaveTableImporter/Sampler classes for use by WaveTable editor and MetaSound nodes (and potentially other systems)
#rnx

diff --git a/Engine/Plugins/Runtime/WaveTable/Source/WaveTable/Private/WaveTableImporter.cpp b/Engine/Plugins/Runtime/WaveTable/Source/WaveTable/Private/WaveTableImporter.cpp
new file mode 100644
--- /dev/null
+++ b/Engine/Plugins/Runtime/WaveTable/Source/WaveTable/Private/WaveTableImporter.cpp
@@ -0,0 +1,114 @@
+// Copyright Epic Games, Inc. All Rights Reserved.
+#include "WaveTableImporter.h"
+
+#include "DSP/FloatArrayMath.h"
+#include "HAL/Platform.h"
+#include "Misc/AssertionMacros.h"
+#include "UObject/SoftObjectPath.h"
+#include "UObject/Object.h"
+#include "UObject/ObjectMacros.h"
+#include "UObject/Package.h"
+#include "WaveTableSampler.h"
+
+
+namespace WaveTable
+{
+	namespace ImporterPrivate
+	{
+		UE_NODISCARD static FORCEINLINE int32 RatioToIndex(float InRatio, int32 InArrayNum)
+		{
+			check(InArrayNum > 0);
+			check(InRatio >= 0.0f);
+
+			return FMath::Min(InArrayNum - 1, FMath::TruncToInt32(InRatio * InArrayNum));
+		};
+	}
+
+	FImporter::FImporter(const FWaveTableSettings& InOptions, EWaveTableResolution InResolution)
+		: Options(InOptions)
+		, Resolution(InResolution)
+	{
+		Sampler.SetInterpolationMode(FWaveTableSampler::EInterpolationMode::Cubic);
+		Sampler.SetPhase(Options.Phase);
+	}
+
+	void FImporter::Process(TArray<float>& OutWaveTable)
+	{
+		using namespace ImporterPrivate;
+
+		OutWaveTable.Reset();
+		OutWaveTable.AddZeroed(ResolutionToInt32(Resolution));
+
+		if (Resolution == EWaveTableResolution::None)
+		{
+			return;
+		}
+
+		int32 SourceNumSamples = Options.SourcePCMData.Num();
+		if (SourceNumSamples == 0)
+		{
+			return;
+		}
+
+		// 1. Use Top & Tail to determine start offset & sample count of SourceArrayView
+		int32 SourceTopOffset = 0;
+		if (Options.Top > 0.0f)
+		{
+			SourceTopOffset = ImporterPrivate::RatioToIndex(Options.Top, SourceNumSamples);
+			SourceNumSamples -= SourceTopOffset;
+		}
+
+		if (Options.Tail > 0.0f)
+		{
+			const int32 SourceTailOffset = RatioToIndex(Options.Tail, SourceNumSamples);
+			SourceNumSamples -= SourceTailOffset;
+		}
+
+		if (SourceNumSamples > 0 && SourceTopOffset < Options.SourcePCMData.Num() && !OutWaveTable.IsEmpty())
+		{
+			TArrayView<const float> SourceArrayView(Options.SourcePCMData.GetData() + SourceTopOffset, SourceNumSamples);
+			TArrayView<float> WaveTableView(OutWaveTable);
+			Sampler.Process(SourceArrayView, WaveTableView);
+
+			if (Options.FadeIn > 0.0f)
+			{
+				const int32 FadeLength = RatioToIndex(Options.FadeIn, OutWaveTable.Num());
+				TArrayView<float> FadeView(OutWaveTable.GetData(), FadeLength);
+				Audio::ArrayFade(FadeView, 0.0f, 1.0f);
+			}
+
+			if (Options.FadeOut > 0.0f)
+			{
+				const int32 FadeLastIndex = RatioToIndex(Options.FadeOut, OutWaveTable.Num());
+				const int32 FadeInitIndex = RatioToIndex(1.0f - Options.FadeOut, OutWaveTable.Num());
+				TArrayView<float> FadeView = TArrayView<float>(&OutWaveTable[FadeInitIndex], FadeLastIndex + 1);
+				Audio::ArrayFade(FadeView, 1.0f, 0.0f);
+			}
+
+			if (Options.bRemoveOffset)
+			{
+				const float TableOffset = Audio::ArrayGetAverageValue(OutWaveTable);
+				Audio::ArrayAddConstantInplace(OutWaveTable, -1.0f * TableOffset);
+			}
+
+			if (Options.bConvertToEnvelope)
+			{
+				Audio::ArrayAbsInPlace(OutWaveTable);
+			}
+			else
+			{
+				Audio::ArrayAddConstantInplace(OutWaveTable, 1.0f);
+				Audio::ArrayMultiplyByConstantInPlace(OutWaveTable, 0.5f);
+			}
+
+			if (Options.bNormalize)
+			{
+				const float MaxValue = Audio::ArrayMaxAbsValue(OutWaveTable);
+				if (MaxValue > 0.0f)
+				{
+					Audio::ArrayMultiplyByConstantInPlace(OutWaveTable, 1.0f / MaxValue);
+				}
+			}
+		}
+	} // namespace Importer
+} // namespace WaveTable
diff --git a/Engine/Plugins/Runtime/WaveTable/Source/WaveTable/Private/WaveTableLogging.cpp b/Engine/Plugins/Runtime/WaveTable/Source/WaveTable/Private/WaveTableLogging.cpp
new file mode 100644
--- /dev/null
+++ b/Engine/Plugins/Runtime/WaveTable/Source/WaveTable/Private/WaveTableLogging.cpp
@@ -0,0 +1,4 @@
+// Copyright Epic Games, Inc. All Rights Reserved.
+#include "WaveTableLogging.h"
+
+DEFINE_LOG_CATEGORY(LogWaveTable);
diff --git a/Engine/Plugins/Runtime/WaveTable/Source/WaveTable/Private/WaveTableLogging.h b/Engine/Plugins/Runtime/WaveTable/Source/WaveTable/Private/WaveTableLogging.h
new file mode 100644
--- /dev/null
+++ b/Engine/Plugins/Runtime/WaveTable/Source/WaveTable/Private/WaveTableLogging.h
@@ -0,0 +1,7 @@
+// Copyright Epic Games, Inc. All Rights Reserved.
+#pragma once
+
+#include "Logging/LogMacros.h"
+
+
+DECLARE_LOG_CATEGORY_EXTERN(LogWaveTable, Log, All);
diff --git a/Engine/Plugins/Runtime/WaveTable/Source/WaveTable/Private/WaveTableModule.cpp b/Engine/Plugins/Runtime/WaveTable/Source/WaveTable/Private/WaveTableModule.cpp
new file mode 100644
--- /dev/null
+++ b/Engine/Plugins/Runtime/WaveTable/Source/WaveTable/Private/WaveTableModule.cpp
@@ -0,0 +1,20 @@
+// Copyright Epic Games, Inc. All Rights Reserved.
+
+#include "WaveTableModule.h"
+
+#include "WaveTableLogging.h"
+
+
+namespace WaveTable
+{
+	void FModule::StartupModule()
+	{
+	}
+
+	void FModule::ShutdownModule()
+	{
+	}
+} // namespace WaveTable
+
+
+IMPLEMENT_MODULE(WaveTable::FModule, WaveTable);
diff --git a/Engine/Plugins/Runtime/WaveTable/Source/WaveTable/Private/WaveTableSampler.cpp b/Engine/Plugins/Runtime/WaveTable/Source/WaveTable/Private/WaveTableSampler.cpp
new file mode 100644
--- /dev/null
+++ b/Engine/Plugins/Runtime/WaveTable/Source/WaveTable/Private/WaveTableSampler.cpp
@@ -0,0 +1,215 @@
+// Copyright Epic Games, Inc. All Rights Reserved.
+#include "WaveTableSampler.h"
+
+#include "DSP/Dsp.h"
+#include "DSP/FloatArrayMath.h"
+
+
+namespace WaveTable
+{
+	namespace UtilitiesPrivate
+	{
+		UE_NODISCARD FORCEINLINE FVector2f GetTangentP0(const float* InTableView, const float NextValue, float InIndex, int32 InArraySize)
+		{
+			const int32 Index = FMath::TruncToInt32(InIndex);
+			const float LastValue = InTableView[(Index - 1) % InArraySize];
+			return FVector2f(2.0f, NextValue - LastValue).GetSafeNormal();
+		};
+
+		UE_NODISCARD FORCEINLINE FVector2f GetTangentP1(const float* InTableView, const float LastValue, float InIndex, int32 InArraySize)
+		{
+			const int32 Index = FMath::TruncToInt32(InIndex);
+			const float NextValue = InTableView[(Index + 1) % InArraySize];
+			return FVector2f(2.0f, NextValue - LastValue).GetSafeNormal();
+		};
+
+		using FInlineInterpolator = TFunctionRef<void(TArrayView<const float> InTableView, TArrayView<float> OutSamplesView)>;
+
+		auto CubicIndexInterpolator = [](TArrayView<const float> InTableView, TArrayView<float> InOutIndicesToValues)
+		{
+			const int32 NumTableSamples = InTableView.Num();
+			const float* InTable = InTableView.GetData();
+			for (int32 i = 0; i < InOutIndicesToValues.Num(); ++i)
+			{
+				float& IndexToOutput = InOutIndicesToValues[i];
+				const int32 LastIndexInt = FMath::TruncToInt32(IndexToOutput);
+
+				const FVector2f P0 = { (float)LastIndexInt, InTableView[LastIndexInt % NumTableSamples] };
+				const FVector2f P1 = { (float)(LastIndexInt + 1), InTable[(LastIndexInt + 1) % NumTableSamples] };
+				const FVector2f TangentP0 = GetTangentP0(InTable, P1.Y, IndexToOutput, NumTableSamples);
+				const FVector2f TangentP1 = GetTangentP1(InTable, P0.Y, IndexToOutput, NumTableSamples);
+				const float Alpha = IndexToOutput - LastIndexInt;
+
+				IndexToOutput = FMath::CubicInterp(P0, TangentP0, P1, TangentP1, Alpha).Y;
+			}
+		};
+
+		auto LinearIndexInterpolator = [](TArrayView<const float> InTableView, const TArrayView<float> InOutIndicesToValues)
+		{
+			const int32 NumTableSamples = InTableView.Num();
+			const float* InTable = InTableView.GetData();
+			for (int32 i = 0; i < InOutIndicesToValues.Num(); ++i)
+			{
+				float& IndexToOutput = InOutIndicesToValues[i];
+				const int32 LastIndexInt = FMath::TruncToInt32(IndexToOutput);
+
+				const float P0 = InTable[LastIndexInt % NumTableSamples];
+				const float Alpha = IndexToOutput - LastIndexInt;
+				const float P1 = InTable[(LastIndexInt + 1) % NumTableSamples];
+
+				IndexToOutput = FMath::Lerp(P0, P1, Alpha);
+			}
+		};
+
+		auto StepIndexInterpolator = [](TArrayView<const float> InTableView, const TArrayView<float> InOutIndicesToValues)
+		{
+			const int32 NumTableSamples = InTableView.Num();
+			const float* InTable = InTableView.GetData();
+			float* IndexToValue = InOutIndicesToValues.GetData();
+			for (int32 i = 0; i < InOutIndicesToValues.Num(); ++i)
+			{
+				float& IndexToOutput = IndexToValue[i];
+				const int32 LastIndexInt = FMath::TruncToInt32(IndexToOutput);
+				IndexToOutput = InTable[LastIndexInt % NumTableSamples];
+			}
+		};
+	}
+
+	FWaveTableSampler::FWaveTableSampler()
+	{
+	}
+
+	FWaveTableSampler::FWaveTableSampler(FSettings&& InSettings)
+		: Settings(MoveTemp(InSettings))
+	{
+	}
+
+	void FWaveTableSampler::SetFreq(float InFreq)
+	{
+		Settings.Freq = InFreq;
+	}
+
+	void FWaveTableSampler::SetInterpolationMode(EInterpolationMode InMode)
+	{
+		Settings.InterpolationMode = InMode;
+	}
+
+	void FWaveTableSampler::SetPhase(float InPhase)
+	{
+		Settings.Phase = InPhase;
+	}
+
+	void FWaveTableSampler::Process(TArrayView<const float> InTableView, TArrayView<float> OutSamplesView)
+	{
+		Process(InTableView, { }, { }, OutSamplesView);
+	}
+
+	void FWaveTableSampler::Process(TArrayView<const float> InTableView, TArrayView<const float> InFreqModulator, TArrayView<const float> InPhaseModulator, TArrayView<float> OutSamplesView)
+	{
+		using namespace UtilitiesPrivate;
+
+		if (!OutSamplesView.IsEmpty())
+		{
+			if (InTableView.IsEmpty())
+			{
+				Audio::ArraySetToConstantInplace(OutSamplesView, 0.0f);
+			}
+			else
+			{
+				ComputeIndexFrequency(InTableView, InFreqModulator, OutSamplesView);
+				ComputeIndexPhase(InTableView, InPhaseModulator, OutSamplesView);
+
+				switch (Settings.InterpolationMode)
+				{
+					case EInterpolationMode::None:
+					{
+						StepIndexInterpolator(InTableView, OutSamplesView);
+					}
+					break;
+
+					case EInterpolationMode::Linear:
+					{
+						LinearIndexInterpolator(InTableView, OutSamplesView);
+					}
+					break;
+
+					case EInterpolationMode::Cubic:
+					{
+						CubicIndexInterpolator(InTableView, OutSamplesView);
+					}
+					break;
+
+					default:
+					{
+						static_assert(static_cast<int32>(EInterpolationMode::COUNT) == 3, "Possible missing switch coverage for EInterpolationMode");
+						checkNoEntry();
+					}
+					break;
+				}
+			}
+
+			if (!FMath::IsNearlyZero(Settings.Offset))
+			{
+				Audio::ArrayAddConstantInplace(OutSamplesView, Settings.Offset);
+			}
+
+			if (!FMath::IsNearlyEqual(Settings.Amplitude, 1.0f))
+			{
+				Audio::ArrayMultiplyByConstantInPlace(OutSamplesView, Settings.Amplitude);
+			}
+		}
+	}
+
+	void FWaveTableSampler::ComputeIndexFrequency(TArrayView<const float> InTableView, TArrayView<const float> InFreqModulator, TArrayView<float> OutSamplesView) const
+	{
+		const float Delta = 1.0f / OutSamplesView.Num();
+		if (InFreqModulator.Num() == OutSamplesView.Num())
+		{
+			FMemory::Memcpy(OutSamplesView.GetData(), InFreqModulator.GetData(), sizeof(float) * OutSamplesView.Num());
+			Audio::ArrayMultiplyByConstantInPlace(OutSamplesView, Settings.Freq * Delta);
+		}
+		else
+		{
+			checkf(InFreqModulator.IsEmpty(), TEXT("FrequencyModulator view should be the same size as the sample view or not supplied (size of 0)."));
+			Audio::ArraySetToConstantInplace(OutSamplesView, Settings.Freq * Delta);
+		}
+
+		float* OutSamples = OutSamplesView.GetData();
+		for (int32 i = 0; i < OutSamplesView.Num(); ++i)
+		{
+			OutSamples[i] *= i;
+		}
+	}
+
+	void FWaveTableSampler::ComputeIndexPhase(TArrayView<const float> InTableView, TArrayView<const float> InPhaseModulator, TArrayView<float> OutSamplesView)
+	{
+		const float DeltaPhase = FMath::Frac(LastIndex / InTableView.Num());
+		if (InPhaseModulator.Num() == OutSamplesView.Num())
+		{
+			PhaseModScratch.SetNum(OutSamplesView.Num());
+			Audio::ArraySetToConstantInplace(PhaseModScratch, Settings.Phase);
+			Audio::ArrayAddInPlace(InPhaseModulator, PhaseModScratch);
+
+			const float* PhaseModData = PhaseModScratch.GetData();
+			float* OutSamples = OutSamplesView.GetData();
+			for (int32 i = 0; i < OutSamplesView.Num(); ++i)
+			{
+				const float Phase = FMath::Frac(DeltaPhase + PhaseModData[i]);
+				const float Offset = OutSamples[i] + Phase;
+				OutSamples[i] = Offset * InTableView.Num();
+			}
+
+			LastIndex = OutSamplesView.Last() - (PhaseModScratch.Last() * InTableView.Num());
+		}
+		else
+		{
+			checkf(InPhaseModulator.IsEmpty(), TEXT("PhaseModulator view should be the same size as the sample view or not supplied (size of 0)."));
+
+			const float Phase = FMath::Frac(DeltaPhase + Settings.Phase);
+			Audio::ArrayAddConstantInplace(OutSamplesView, Phase);
+			Audio::ArrayMultiplyByConstantInPlace(OutSamplesView, InTableView.Num());
+
+			LastIndex = OutSamplesView.Last() - (Settings.Phase * InTableView.Num());
+		}
+	}
+} // namespace WaveTable
diff --git a/Engine/Plugins/Runtime/WaveTable/Source/WaveTable/Private/WaveTableSettings.cpp b/Engine/Plugins/Runtime/WaveTable/Source/WaveTable/Private/WaveTableSettings.cpp
new file mode 100644
--- /dev/null
+++ b/Engine/Plugins/Runtime/WaveTable/Source/WaveTable/Private/WaveTableSettings.cpp
@@ -0,0 +1,60 @@
+// Copyright Epic Games, Inc. All Rights Reserved.
+#include "WaveTableSettings.h"
+
+
+namespace WaveTable
+{
+	int32 ResolutionToInt32(EWaveTableResolution InResolution)
+	{
+		if (InResolution == EWaveTableResolution::None)
+		{
+			return 0;
+		}
+
+		return 1 << static_cast<int32>(InResolution);
+	}
+
+	int32 ResolutionToInt32(EWaveTableCurve InCurve, EWaveTableResolution FileResolution)
+	{
+		EWaveTableResolution Resolution = EWaveTableResolution::None;
+		switch (InCurve)
+		{
+			// File uses the incoming resolution
+			case EWaveTableCurve::File:
+			{
+				Resolution = FileResolution;
+			}
+			break;
+
+			case EWaveTableCurve::Linear:
+			{
+				Resolution = EWaveTableResolution::Res_16;
+			}
+			break;
+
+			case EWaveTableCurve::Exp:
+			case EWaveTableCurve::Exp_Inverse:
+			case EWaveTableCurve::Log:
+			{
+				Resolution = EWaveTableResolution::Res_256;
+			}
+			break;
+
+			case EWaveTableCurve::Sin:
+			case EWaveTableCurve::SCurve:
+			{
+				Resolution = EWaveTableResolution::Res_64;
+			}
+			break;
+
+
+			default:
+			{
+				Resolution = EWaveTableResolution::Res_128;
+			}
+			break;
+		}
+
+		return ResolutionToInt32(Resolution);
+	}
+} // namespace WaveTable
diff --git a/Engine/Plugins/Runtime/WaveTable/Source/WaveTable/Public/WaveTableImporter.h b/Engine/Plugins/Runtime/WaveTable/Source/WaveTable/Public/WaveTableImporter.h
new file mode 100644
--- /dev/null
+++ b/Engine/Plugins/Runtime/WaveTable/Source/WaveTable/Public/WaveTableImporter.h
@@ -0,0 +1,23 @@
+// Copyright Epic Games, Inc. All Rights Reserved.
+#pragma once
+
+#include "Containers/Array.h"
+#include "WaveTableSettings.h"
+#include "WaveTableSampler.h"
+
+
+namespace WaveTable
+{
+	class WAVETABLE_API FImporter
+	{
+	public:
+		FImporter(const FWaveTableSettings& InOptions, EWaveTableResolution InResolution);
+
+		void Process(TArray<float>& OutWaveTable);
+
+	private:
+		const FWaveTableSettings& Options;
+		EWaveTableResolution Resolution;
+		FWaveTableSampler Sampler;
+	};
+} // namespace WaveTable
diff --git a/Engine/Plugins/Runtime/WaveTable/Source/WaveTable/Public/WaveTableModule.h b/Engine/Plugins/Runtime/WaveTable/Source/WaveTable/Public/WaveTableModule.h
new file mode 100644
--- /dev/null
+++ b/Engine/Plugins/Runtime/WaveTable/Source/WaveTable/Public/WaveTableModule.h
@@ -0,0 +1,15 @@
+// Copyright Epic Games, Inc. All Rights Reserved.
+#pragma once
+
+#include "Modules/ModuleInterface.h"
+
+
+namespace WaveTable
+{
+class WAVETABLE_API FModule : public IModuleInterface
+{
+public:
+	virtual void StartupModule() override;
+	virtual void ShutdownModule() override;
+};
+} // namespace WaveTable
\ No newline at end of file
diff --git a/Engine/Plugins/Runtime/WaveTable/Source/WaveTable/Public/WaveTableSampler.h b/Engine/Plugins/Runtime/WaveTable/Source/WaveTable/Public/WaveTableSampler.h
new file mode 100644
--- /dev/null
+++ b/Engine/Plugins/Runtime/WaveTable/Source/WaveTable/Public/WaveTableSampler.h
@@ -0,0 +1,58 @@
+// Copyright Epic Games, Inc. All Rights Reserved.
+#pragma once
+
+#include "Containers/Array.h"
+#include "DSP/Filter.h"
+#include "HAL/Platform.h"
+#include "WaveTableSettings.h"
+
+
+namespace WaveTable
+{
+	class WAVETABLE_API FWaveTableSampler
+	{
+	public:
+		enum class EInterpolationMode : uint8
+		{
+			None,
+			Linear,
+			Cubic,
+
+			COUNT
+		};
+
+		struct WAVETABLE_API FSettings
+		{
+			float Amplitude = 1.0f;
+			float Offset = 0.0f;
+
+			// How many times to read through the incoming
+			// buffer as a ratio to a single output buffer
+			float Freq = 1.0f;
+
+			// Offset from beginning to end of array [0.0, 1.0]
+			float Phase = 0.0f;
+
+			EInterpolationMode InterpolationMode = EInterpolationMode::Linear;
+		};
+
+		FWaveTableSampler();
+		FWaveTableSampler(FSettings&& InSettings);
+
+		void Process(TArrayView<const float> InTableView, TArrayView<float> OutSamplesView);
+		void Process(TArrayView<const float> InTableView, TArrayView<const float> InFreqModulator, TArrayView<const float> InPhaseModulator, TArrayView<float> OutSamplesView);
+
+		void SetInterpolationMode(EInterpolationMode InMode);
+		void SetFreq(float InFreq);
+		void SetPhase(float InPhase);
+
+	private:
+		void ComputeIndexFrequency(TArrayView<const float> InTableView, TArrayView<const float> InFreqModulator, TArrayView<float> OutSamplesView) const;
+		void ComputeIndexPhase(TArrayView<const float> InTableView, TArrayView<const float> InPhaseModulator, TArrayView<float> OutSamplesView);
+
+		float LastIndex = 0.0;
+
+		TArray<float> PhaseModScratch;
+		FSettings Settings;
+	};
+} // namespace WaveTable
\ No newline at end of file
diff --git a/Engine/Plugins/Runtime/WaveTable/Source/WaveTable/Public/WaveTableSettings.h b/Engine/Plugins/Runtime/WaveTable/Source/WaveTable/Public/WaveTableSettings.h
new file mode 100644
--- /dev/null
+++ b/Engine/Plugins/Runtime/WaveTable/Source/WaveTable/Public/WaveTableSettings.h
@@ -0,0 +1,109 @@
+// Copyright Epic Games, Inc. All Rights Reserved.
+#pragma once
+
+#include "Containers/Array.h"
+#include "HAL/Platform.h"
+#include "UObject/SoftObjectPath.h"
+#include "UObject/Object.h"
+#include "UObject/ObjectMacros.h"
+
+#include "WaveTableSettings.generated.h"
+
+
+UENUM(BlueprintType)
+enum class EWaveTableResolution : uint8
+{
+	None		= 0 UMETA(DisplayName = "None"),
+	Res_16		= 4 UMETA(DisplayName = "16"),
+	Res_32		= 5 UMETA(DisplayName = "32"),
+	Res_64		= 6 UMETA(DisplayName = "64"),
+	Res_128		= 7 UMETA(DisplayName = "128"),
+	Res_256		= 8 UMETA(DisplayName = "256"),
+	Res_512		= 9 UMETA(DisplayName = "512"),
+	Res_1024	= 10 UMETA(DisplayName = "1024"),
+	Res_2048	= 11 UMETA(DisplayName = "2048"),
+	Res_4096	= 12 UMETA(DisplayName = "4096"),
+
+	MAX = Res_4096 UMETA(Hidden)
+};
+
+
+USTRUCT()
+struct WAVETABLE_API FWaveTableSettings
+{
+	GENERATED_USTRUCT_BODY()
+
+	// File to import
+	UPROPERTY(EditAnywhere, Category = Options)
+	FFilePath FilePath;
+
+	// Index of channel in file to build WaveTable from (wraps if channel is greater than number in file)
+	UPROPERTY(EditAnywhere, Category = Options, meta = (ClampMin = "0.0"))
+	int32 ChannelIndex = 0;
+
+	// Percent to phase shift of table
+	UPROPERTY(EditAnywhere, Category = Options, meta = (ClampMin = "0.0", ClampMax = "1.0", SliderExponent = 0.5))
+	float Phase = 0.0f;
+
+	// Percent to remove from beginning of sampled WaveTable.
+	UPROPERTY(EditAnywhere, Category = Options, meta = (ClampMin = "0.0", ClampMax = "1.0", SliderExponent = 0.5))
+	float Top = 0.0f;
+
+	// Percent to remove from end of sampled WaveTable.
+	UPROPERTY(EditAnywhere, Category = Options, meta = (ClampMin = "0.0", ClampMax = "1.0", SliderExponent = 0.5))
+	float Tail = 0.0f;
+
+	// Percent to fade in over.
+	UPROPERTY(EditAnywhere, Category = Options, meta = (ClampMin = "0.0", ClampMax = "1.0", SliderExponent = 0.5))
+	float FadeIn  = 0.0f;
+
+	// Percent to fade out over.
+	UPROPERTY(EditAnywhere, Category = Options, meta = (ClampMin = "0.0", ClampMax = "1.0", SliderExponent = 0.5))
+	float FadeOut = 0.0f;
+
+	// Whether or not to normalize the WaveTable.
+	UPROPERTY(EditAnywhere, Category = Options)
+	bool bNormalize = true;
+
+	// If true, converts to envelope, effectively taking
+	// the absolute value of the given file's waveform.
+	UPROPERTY(EditAnywhere, Category = Options)
+	bool bConvertToEnvelope = false;
+
+	// Whether or not to remove offset from original file
+	// (analogous to "DC offset" in circuit theory).
+	UPROPERTY(EditAnywhere, Category = Options)
+	bool bRemoveOffset = false;
+
+	// SourcePCM Data
+	UPROPERTY()
+	TArray<float> SourcePCMData;
+};
+
+UENUM(BlueprintType)
+enum class EWaveTableCurve : uint8
+{
+	Linear		UMETA(DisplayName = "Linear"),
+	Exp			UMETA(DisplayName = "Exponential"),
+	Exp_Inverse UMETA(DisplayName = "Exponential (Inverse)"),
+	Log			UMETA(DisplayName = "Log"),
+	Sin			UMETA(DisplayName = "Sin (Quarter)"),
+	SCurve		UMETA(DisplayName = "Sin (S-Curve)"),
+
+	// Reference a shared curve asset
+	Shared		UMETA(DisplayName = "Shared"),
+
+	// Design a custom curve unique to the owning transform
+	Custom		UMETA(DisplayName = "Custom"),
+
+	// Generate WaveTable from audio file
+	File		UMETA(DisplayName = "File"),
+
+	Count		UMETA(Hidden),
+};
+
+namespace WaveTable
+{
+	WAVETABLE_API int32 ResolutionToInt32(EWaveTableResolution InResolution);
+	WAVETABLE_API int32 ResolutionToInt32(EWaveTableCurve InCurve, EWaveTableResolution FileResolution);
+} // namespace WaveTable
\ No newline at end of file
diff --git a/Engine/Plugins/Runtime/WaveTable/Source/WaveTable/WaveTable.Build.cs b/Engine/Plugins/Runtime/WaveTable/Source/WaveTable/WaveTable.Build.cs
new file mode 100644
--- /dev/null
+++ b/Engine/Plugins/Runtime/WaveTable/Source/WaveTable/WaveTable.Build.cs
@@ -0,0 +1,27 @@
+// Copyright Epic Games, Inc. All Rights Reserved.
+
+namespace UnrealBuildTool.Rules
+{
+	public class WaveTable : ModuleRules
+	{
+		public WaveTable(ReadOnlyTargetRules Target) : base(Target)
+		{
+			OptimizeCode = CodeOptimization.Never;
+
+			PublicDependencyModuleNames.AddRange(
+				new string[]
+				{
+					"Core"
+				}
+			);
+
+			PrivateDependencyModuleNames.AddRange(
+				new string[]
+				{
+					"CoreUObject",
+					"SignalProcessing"
+				}
+			);
+		}
+	}
+}
\ No newline at end of file
diff --git a/Engine/Plugins/Runtime/WaveTable/Source/WaveTableEditor/Private/WaveTableBankEditor.cpp b/Engine/Plugins/Runtime/WaveTable/Source/WaveTableEditor/Private/WaveTableBankEditor.cpp
new file mode 100644
--- /dev/null
+++ b/Engine/Plugins/Runtime/WaveTable/Source/WaveTableEditor/Private/WaveTableBankEditor.cpp
@@ -0,0 +1,478 @@
+// Copyright Epic Games, Inc. All Rights Reserved.
+#include "WaveTableBankEditor.h"
+
+#include "CommonFrameRates.h"
+#include "Containers/Set.h"
+#include "Curves/CurveFloat.h"
+#include "Curves/CurveLinearColor.h"
+#include "CurveEditor.h"
+#include "CurveEditorCommands.h"
+#include "Styling/AppStyle.h"
+#include "Framework/Commands/UIAction.h"
+#include "Framework/MultiBox/MultiBoxDefs.h"
+#include "Framework/MultiBox/MultiBoxBuilder.h"
+#include "Framework/MultiBox/MultiBoxExtender.h"
+#include "IDetailsView.h"
+#include "Modules/ModuleManager.h"
+#include "PropertyEditorModule.h"
+#include "RichCurveEditorModel.h"
+#include "SCurveEditorPanel.h"
+#include "Tree/SCurveEditorTree.h"
+#include "Tree/ICurveEditorTreeItem.h"
+#include "Tree/SCurveEditorTreePin.h"
+#include "WaveTableSettings.h"
+#include "WaveTableSampler.h"
+#include "Widgets/Docking/SDockTab.h"
+#include "Widgets/Input/SButton.h"
+#include "Widgets/Input/SNumericDropDown.h"
+#include "Widgets/Layout/SBorder.h"
+#include "Widgets/SFrameRatePicker.h"
+
+
+#define LOCTEXT_NAMESPACE "WaveTableEditor"
+
+
+
+namespace WaveTable
+{
+	namespace Editor
+	{
+		const FName FWaveTableBankEditor::AppIdentifier(TEXT("WaveTableEditorApp"));
+		const FName FWaveTableBankEditor::CurveTabId(TEXT("WaveTableEditor_Curves"));
+		const FName FWaveTableBankEditor::PropertiesTabId(TEXT("WaveTableEditor_Properties"));
+
+		FWaveTableBankEditor::FWaveTableBankEditor()
+		{
+		}
+
+		void FWaveTableBankEditor::RegisterTabSpawners(const TSharedRef<FTabManager>& InTabManager)
+		{
+			WorkspaceMenuCategory = InTabManager->AddLocalWorkspaceMenuCategory(LOCTEXT("WorkspaceMenu_WaveTableEditor", "WaveTable Editor"));
+
+			FAssetEditorToolkit::RegisterTabSpawners(InTabManager);
+
+			InTabManager->RegisterTabSpawner(PropertiesTabId, FOnSpawnTab::CreateSP(this, &FWaveTableBankEditor::SpawnTab_Properties))
+				.SetDisplayName(LOCTEXT("DetailsTab", "Details"))
+				.SetGroup(WorkspaceMenuCategory.ToSharedRef())
+				.SetIcon(FSlateIcon(FAppStyle::GetAppStyleSetName(), "LevelEditor.Tabs.Details"));
+
+			FSlateIcon CurveIcon(FAppStyle::GetAppStyleSetName(), "WaveTableEditor.Tabs.Properties");
+			InTabManager->RegisterTabSpawner(CurveTabId, FOnSpawnTab::CreateLambda([this](const FSpawnTabArgs& Args) { return SpawnTab_OutputCurve(Args); }))
+				.SetDisplayName(LOCTEXT("TransformCurvesTab", "Transform Curves"))
+				.SetGroup(WorkspaceMenuCategory.ToSharedRef())
+				.SetIcon(CurveIcon);
+		}
+
+		void FWaveTableBankEditor::UnregisterTabSpawners(const TSharedRef<FTabManager>& InTabManager)
+		{
+			InTabManager->UnregisterTabSpawner(PropertiesTabId);
+			InTabManager->UnregisterTabSpawner(CurveTabId);
+		}
+
+		void FWaveTableBankEditor::Init(const EToolkitMode::Type Mode, const TSharedPtr<IToolkitHost>& InitToolkitHost, UObject* InParentObject)
+		{
+			check(InParentObject);
+
+			CurveEditor = MakeShared<FCurveEditor>();
+			FCurveEditorInitParams InitParams;
+			CurveEditor->InitCurveEditor(InitParams);
+			CurveEditor->GridLineLabelFormatXAttribute = LOCTEXT("GridXLabelFormat", "{0}");
+
+			TUniquePtr<ICurveEditorBounds> EditorBounds = MakeUnique<FStaticCurveEditorBounds>();
+			EditorBounds->SetInputBounds(0.05, 1.05);
+			CurveEditor->SetBounds(MoveTemp(EditorBounds));
+
+			CurvePanel = SNew(SCurveEditorPanel, CurveEditor.ToSharedRef());
+
+			// Support undo/redo
+			InParentObject->SetFlags(RF_Transactional);
+			GEditor->RegisterForUndo(this);
+
+			FDetailsViewArgs Args;
+			Args.bHideSelectionTip = true;
+			Args.NotifyHook = this;
+
+			FPropertyEditorModule& PropertyModule = FModuleManager::LoadModuleChecked<FPropertyEditorModule>("PropertyEditor");
+			PropertiesView = PropertyModule.CreateDetailView(Args);
+			PropertiesView->SetObject(InParentObject);
+
+			TSharedRef<FTabManager::FLayout> StandaloneDefaultLayout = FTabManager::NewLayout("Standalone_WaveTableEditor_Layout_v1")
+				->AddArea
+				(
+					FTabManager::NewPrimaryArea()
+					->SetOrientation(Orient_Vertical)
+					->Split
+					(
+						FTabManager::NewSplitter()
+						->SetSizeCoefficient(0.9f)
+						->SetOrientation(Orient_Horizontal)
+						->Split
+						(
+							FTabManager::NewStack()
+							->SetSizeCoefficient(0.225f)
+							->AddTab(PropertiesTabId, ETabState::OpenedTab)
+						)
+						->Split
+						(
+							FTabManager::NewSplitter()
+							->SetSizeCoefficient(0.775f)
+							->SetOrientation(Orient_Vertical)
+							->Split
+							(
+								FTabManager::NewStack()
+								->SetHideTabWell(true)
+								->SetSizeCoefficient(0.33f)
+								->AddTab(CurveTabId, ETabState::OpenedTab)
+							)
+						)
+					)
+				);
+
+			const bool bCreateDefaultStandaloneMenu = true;
+			const bool bCreateDefaultToolbar = true;
+			const bool bToolbarFocusable = false;
+			const bool bUseSmallIcons = true;
+			FAssetEditorToolkit::InitAssetEditor(
+				Mode,
+				InitToolkitHost,
+				AppIdentifier,
+				StandaloneDefaultLayout,
+				bCreateDefaultStandaloneMenu,
+				bCreateDefaultToolbar,
+				InParentObject,
+				bToolbarFocusable,
+				bUseSmallIcons);
+
+		// 	Editor::FModule& WaveTableModule = FModuleManager::LoadModuleChecked<Editor::FModule>(TEXT("WaveTableEditor"));
+		// 	AddMenuExtender(WaveTableModule.GetWaveTableMenuExtensibilityManager()->GetAllExtenders(GetToolkitCommands(), GetEditingObjects()));
+		// 	AddToolbarExtender(WaveTableModule.GetWaveTableToolbarExtensibilityManager()->GetAllExtenders(GetToolkitCommands(), GetEditingObjects()));
+			AddToolbarExtender(CurvePanel->GetToolbarExtender());
+
+			if (CurveEditor.IsValid())
+			{
+				RegenerateMenusAndToolbars();
+			}
+		}
+
+		void FWaveTableBankEditor::ClearExpressionCurve(int32 InCurveIndex)
+		{
+			if (CurveData.IsValidIndex(InCurveIndex))
+			{
+				CurveData[InCurveIndex].ExpressionCurve.Reset();
+			}
+		}
+
+		bool FWaveTableBankEditor::RequiresNewCurve(int32 InCurveIndex, const FRichCurve& InRichCurve) const
+		{
+			const FCurveModelID CurveModelID = CurveData[InCurveIndex].ModelID;
+			const TUniquePtr<FCurveModel>* CurveModel = CurveEditor->GetCurves().Find(CurveModelID);
+			if (!CurveModel || !CurveModel->IsValid())
+			{
+				return true;
+			}
+
+			FWaveTableCurveModelBase* PatchCurveModel = static_cast<FWaveTableCurveModelBase*>(CurveModel->Get());
+			check(PatchCurveModel);
+			if (&PatchCurveModel->GetRichCurve() != &InRichCurve)
+			{
+				return true;
+			}
+
+			return false;
+		}
+
+		void FWaveTableBankEditor::SetCurve(int32 InCurveIndex, FRichCurve& InRichCurve, EWaveTableCurveSource InSource)
+		{
+			check(CurveEditor.IsValid());
+
+			if (!ensure(CurveData.IsValidIndex(InCurveIndex)))
+			{
+				return;
+			}
+
+			FCurveData& CurveDataEntry = CurveData[InCurveIndex];
+
+			const bool bRequiresNewCurve = RequiresNewCurve(InCurveIndex, InRichCurve);
+			if (bRequiresNewCurve)
+			{
+				TUniquePtr<FWaveTableCurveModelBase> NewCurve = ConstructCurveModel(InRichCurve, GetEditingObject(), InSource);
+				NewCurve->Refresh(InCurveIndex);
+				CurveDataEntry.ModelID = CurveEditor->AddCurve(MoveTemp(NewCurve));
+			}
+			else
+			{
+				const TUniquePtr<FCurveModel>& CurveModel = CurveEditor->GetCurves().FindChecked(CurveDataEntry.ModelID);
+				check(CurveModel.Get());
+				static_cast<FWaveTableCurveModelBase*>(CurveModel.Get())->Refresh(InCurveIndex);
+			}
+
+			CurveEditor->PinCurve(CurveDataEntry.ModelID);
+		}
+
+		FName FWaveTableBankEditor::GetToolkitFName() const
+		{
+			return FName("WaveTableEditor");
+		}
+
+		FText FWaveTableBankEditor::GetBaseToolkitName() const
+		{
+			return LOCTEXT( "AppLabel", "WaveTable Editor" );
+		}
+
+		FString FWaveTableBankEditor::GetWorldCentricTabPrefix() const
+		{
+			return LOCTEXT("WorldCentricTabPrefix", "WaveTable ").ToString();
+		}
+
+		FLinearColor FWaveTableBankEditor::GetWorldCentricTabColorScale() const
+		{
+			return FLinearColor( 0.0f, 0.0f, 0.2f, 0.5f );
+		}
+
+		EOrientation FWaveTableBankEditor::GetSnapLabelOrientation() const
+		{
+			return FMultiBoxSettings::UseSmallToolBarIcons.Get()
+				? EOrientation::Orient_Horizontal
+				: EOrientation::Orient_Vertical;
+		}
+
+		void FWaveTableBankEditor::NotifyPostChange(const FPropertyChangedEvent& PropertyChangedEvent, FProperty* PropertyThatChanged)
+		{
+			if (PropertyChangedEvent.ChangeType != EPropertyChangeType::Interactive)
+			{
+				RefreshCurves();
+			}
+		}
+
+		TSharedRef<SDockTab> FWaveTableBankEditor::SpawnTab_Properties(const FSpawnTabArgs& Args)
+		{
+			check(Args.GetTabId() == PropertiesTabId);
+
+			return SNew(SDockTab)
+				.Label(LOCTEXT("SoundWaveTableDetailsTitle", "Details"))
+				[
+					PropertiesView.ToSharedRef()
+				];
+		}
+
+		TSharedRef<SDockTab> FWaveTableBankEditor::SpawnTab_OutputCurve(const FSpawnTabArgs& Args)
+		{
+			RefreshCurves();
+			CurveEditor->ZoomToFit();
+
+			TSharedRef<SDockTab> NewDockTab = SNew(SDockTab)
+				.Label(FText::Format(LOCTEXT("WaveTableFilterTitle", "Filter Transform Curve: {0}"), FText::FromString(GetEditingObject()->GetName())))
+				.TabColorScale(GetTabColorScale())
+				[
+					SNew(SBorder)
+					.BorderImage(FAppStyle::GetBrush("ToolPanel.GroupBorder"))
+					.Padding(0.0f)
+					[
+						CurvePanel.ToSharedRef()
+					]
+				];
+
+				return NewDockTab;
+		}
+
+		void FWaveTableBankEditor::PostUndo(bool bSuccess)
+		{
+			if (bSuccess)
+			{
+				RefreshCurves();
+			}
+		}
+
+		void FWaveTableBankEditor::ResetCurves()
+		{
+			check(CurveEditor.IsValid());
+
+			CurveEditor->RemoveAllCurves();
+			CurveData.Reset();
+			CurveData.AddDefaulted(GetNumCurves());
+		}
+
+		void FWaveTableBankEditor::InitCurves()
+		{
+			for (int32 i = 0; i < GetNumCurves(); ++i)
+			{
+				EWaveTableCurve CurveType = GetCurveType(i);
+				switch (CurveType)
+				{
+					case EWaveTableCurve::Exp:
+					case EWaveTableCurve::Exp_Inverse:
+					case EWaveTableCurve::Linear:
+					case EWaveTableCurve::Log:
+					case EWaveTableCurve::SCurve:
+					case EWaveTableCurve::Sin:
+					case EWaveTableCurve::File:
+					{
+						if (RequiresNewCurve(i, *CurveData[i].ExpressionCurve.Get()))
+						{
+							ResetCurves();
+						}
+					}
+					break;
+
+					case EWaveTableCurve::Shared:
+					{
+						if (const UCurveFloat* SharedCurve = GetSharedCurve(i))
+						{
+							if (RequiresNewCurve(i, SharedCurve->FloatCurve))
+							{
+								ResetCurves();
+							}
+						}
+						else if (RequiresNewCurve(i, *CurveData[i].ExpressionCurve.Get()))
+						{
+							ResetCurves();
+						}
+					}
+					break;
+
+					case EWaveTableCurve::Custom:
+					{
+						if (RequiresNewCurve(i, GetCustomCurveChecked(i)))
+						{
+							ResetCurves();
+						}
+					}
+					break;
+
+					default:
+					{
+						static_assert(static_cast<int32>(EWaveTableCurve::Count) == 9, "Possible missing case coverage for output curve.");
+					}
+					break;
+				}
+			}
+		}
+
+		void FWaveTableBankEditor::RefreshCurves()
+		{
+			check(CurveEditor.IsValid());
+
+			for (const FCurveData& CurveDataEntry : CurveData)
+			{
+				CurveEditor->UnpinCurve(CurveDataEntry.ModelID);
+			}
+
+			if (GetIsPropertyEditorDisabled())
+			{
+				ResetCurves();
+				return;
+			}
+
+			const int32 NumCurves = GetNumCurves();
+			if (NumCurves == CurveData.Num())
+			{
+				InitCurves();
+			}
+			else
+			{
+				ResetCurves();
+			}
+
+			for (int32 i = 0; i < GetNumCurves(); ++i)
+			{
+				const EWaveTableCurve CurveType = GetCurveType(i);
+				if (!CurveData.IsValidIndex(i))
+				{
+					continue;
+				}
+
+				FCurveData& CurveDataEntry = CurveData[i];
+
+				switch (CurveType)
+				{
+					case EWaveTableCurve::Exp:
+					case EWaveTableCurve::Exp_Inverse:
+					case EWaveTableCurve::Linear:
+					case EWaveTableCurve::Log:
+					case EWaveTableCurve::SCurve:
+					case EWaveTableCurve::Sin:
+					case EWaveTableCurve::File:
+					{
+						GenerateExpressionCurve(CurveDataEntry, i, EWaveTableCurveSource::Expression);
+					}
+					break;
+
+					case EWaveTableCurve::Shared:
+					{
+						if (UCurveFloat* SharedCurve = GetSharedCurve(i))
+						{
+							ClearExpressionCurve(i);
+							SetCurve(i, SharedCurve->FloatCurve, EWaveTableCurveSource::Shared);
+						}
+						else
+						{
+							// Builds a dummy expression that just maps input to output in case
+							// where asset isn't selected and leave source as unset
+							GenerateExpressionCurve(CurveDataEntry, i, EWaveTableCurveSource::Expression, true /* bIsUnset */);
+						}
+					}
+					break;
+
+					case EWaveTableCurve::Custom:
+					{
+						FRichCurve& CustomCurve = GetCustomCurveChecked(i);
+						TrimKeys(CustomCurve);
+						ClearExpressionCurve(i);
+						SetCurve(i, CustomCurve, EWaveTableCurveSource::Custom);
+					}
+					break;
+
+					default:
+					{
+						static_assert(static_cast<int32>(EWaveTableCurve::Count) == 9, "Possible missing case coverage for output curve.");
+					}
+					break;
+				}
+			}
+
+			// Collect and remove stale curves from editor
+			TArray<FCurveModelID> ToRemove;
+			TSet<FCurveModelID> ActiveModelIDs;
+			for (const FCurveData& CurveDataEntry : CurveData)
+			{
+				ActiveModelIDs.Add(CurveDataEntry.ModelID);
+			}
+
+			const TMap<FCurveModelID, TUniquePtr<FCurveModel>>& Curves = CurveEditor->GetCurves();
+			for (const TPair<FCurveModelID, TUniquePtr<FCurveModel>>& Pair : Curves)
+			{
+				if (!ActiveModelIDs.Contains(Pair.Key))
+				{
+					ToRemove.Add(Pair.Key);
+				}
+			}
+
+			for (FCurveModelID ModelID : ToRemove)
+			{
+				CurveEditor->RemoveCurve(ModelID);
+			}
+		}
+
+		void FWaveTableBankEditor::PostRedo(bool bSuccess)
+		{
+			if (bSuccess)
+			{
+				RefreshCurves();
+			}
+		}
+
+		void FWaveTableBankEditor::TrimKeys(FRichCurve& OutCurve)
+		{
+			while (OutCurve.GetNumKeys() > 0 && 0.0f > OutCurve.GetFirstKey().Time)
+			{
+				OutCurve.DeleteKey(OutCurve.GetFirstKeyHandle());
+			}
+
+			while (OutCurve.GetNumKeys() > 0 && 1.0f < OutCurve.GetLastKey().Time)
+			{
+				OutCurve.DeleteKey(OutCurve.GetLastKeyHandle());
+			}
+		}
+	} // namespace Editor
+} // namespace WaveTable
+#undef LOCTEXT_NAMESPACE
diff --git a/Engine/Plugins/Runtime/WaveTable/Source/WaveTableEditor/Private/WaveTableCurveEditorViewStacked.cpp b/Engine/Plugins/Runtime/WaveTable/Source/WaveTableEditor/Private/WaveTableCurveEditorViewStacked.cpp
new file mode 100644
--- /dev/null
+++ b/Engine/Plugins/Runtime/WaveTable/Source/WaveTableEditor/Private/WaveTableCurveEditorViewStacked.cpp
@@ -0,0 +1,469 @@
+// Copyright Epic Games, Inc. All Rights Reserved.
+
+#include "WaveTableCurveEditorViewStacked.h"
+
+#include "Curves/CurveFloat.h"
+#include "CurveEditor.h"
+#include "CurveModel.h"
+#include "Fonts/FontMeasure.h"
+#include "Framework/Application/SlateApplication.h"
+#include "SCurveEditorPanel.h"
+#include "SGraphActionMenu.h"
+#include "Styling/AppStyle.h"
+#include "Widgets/Text/STextBlock.h"
+
+#define LOCTEXT_NAMESPACE "WaveTableEditor"
+
+
+namespace WaveTable
+{
+	namespace Editor
+	{
+		FWaveTableCurveModelBase::FWaveTableCurveModelBase(FRichCurve& InRichCurve, UObject* InOwner, EWaveTableCurveSource InSource)
+			: FRichCurveEditorModelRaw(&InRichCurve, InOwner)
+			, ParentObject(InOwner)
+			, Source(InSource)
+		{
+			check(InOwner);
+		}
+
+		void FWaveTableCurveModelBase::Refresh(int32 InCurveIndex)
+		{
+			// Must be set prior to remainder of refresh to avoid
+			// child class implementation acting on incorrect cached index
+			// (Editor May re-purpose models resulting in index change).
+			CurveIndex = InCurveIndex;
+
+			Color = GetCurveColor();
+
+			FText OutputAxisName;
+			RefreshCurveDescriptorText(ShortDisplayName, InputAxisName, OutputAxisName);
+
+			ShortDisplayName = GetShortDisplayName();
+			AxesDescriptor = FText::Format(LOCTEXT("WaveTableCurveDisplayTitle_AxesFormat", "X = {0}, Y = {1}"), InputAxisName, OutputAxisName);
+
+			bKeyDrawEnabled = true;
+			IntentionName = TEXT("AudioControlValue");
+			SupportedViews = ViewId;
+
+			bKeyDrawEnabled = false;
+
+			const bool bIsDisabled = GetPropertyEditorDisabled();
+			if (bIsDisabled)
+			{
+				LongDisplayName = FText::Format(LOCTEXT("WaveTableCurveDisplay_BypassedNameFormat", "{0} ({1})"), ShortDisplayName, GetPropertyEditorDisabledText());
+			}
+			else
+			{
+				UEnum* Enum = StaticEnum<EWaveTableCurveSource>();
+				check(Enum);
+
+				FString EnumValueName = Enum->GetValueAsString(Source);
+				int32 DelimIndex = -1;
+				EnumValueName.FindLastChar(':', DelimIndex);
+				if (DelimIndex > 0 && DelimIndex < EnumValueName.Len() - 1)
+				{
+					EnumValueName.RightChopInline(DelimIndex + 1);
+				}
+
+				FText CurveSourceText = FText::FromString(EnumValueName);
+
+				if (Source == EWaveTableCurveSource::Custom)
+				{
+					bKeyDrawEnabled = true;
+				}
+				else if (Source == EWaveTableCurveSource::Shared)
+				{
+					const UCurveFloat* SharedCurve = GetSharedCurve();
+					check(SharedCurve);
+					CurveSourceText = FText::Format(LOCTEXT("WaveTableCurveDisplayTitle_SharedNameFormat", "Shared, {0}"), FText::FromString(SharedCurve->GetName()));
+					bKeyDrawEnabled = true;
+				}
+
+				LongDisplayName = FText::Format(LOCTEXT("WaveTableCurveDisplayTitle_NameFormat", "{0}: {1} ({2})"), FText::AsNumber(CurveIndex), ShortDisplayName, CurveSourceText);
+			}
+		}
+
+		FLinearColor FWaveTableCurveModelBase::GetColor() const
+		{
+			return !GetPropertyEditorDisabled() && (Source == EWaveTableCurveSource::Custom)
+				? Color
+				: Color.Desaturate(0.35f);
+		}
+
+		const FText& FWaveTableCurveModelBase::GetAxesDescriptor() const
+		{
+			return AxesDescriptor;
+		}
+
+		EWaveTableCurveSource FWaveTableCurveModelBase::GetSource() const
+		{
+			return Source;
+		}
+
+		const UObject* FWaveTableCurveModelBase::GetParentObject() const
+		{
+			return ParentObject.Get();
+		}
+
+		bool FWaveTableCurveModelBase::IsReadOnly() const
+		{
+			return Source != EWaveTableCurveSource::Custom;
+		}
+
+		ECurveEditorViewID FWaveTableCurveModelBase::ViewId = ECurveEditorViewID::Invalid;
+
+		void SViewStacked::Construct(const FArguments& InArgs, TWeakPtr<FCurveEditor> InCurveEditor)
+		{
+			SCurveEditorViewStacked::Construct(InArgs, InCurveEditor);
+
+			StackedHeight = 300.0f;
+			StackedPadding = 60.0f;
+		}
+
+		void SViewStacked::PaintView(const FPaintArgs& Args, const FGeometry& AllottedGeometry, const FSlateRect& MyCullingRect, FSlateWindowElementList& OutDrawElements, int32 BaseLayerId, const FWidgetStyle& InWidgetStyle, bool bParentEnabled) const
+		{
+			TSharedPtr<FCurveEditor> CurveEditor = WeakCurveEditor.Pin();
+			if (CurveEditor)
+			{
+				const ESlateDrawEffect DrawEffects = ShouldBeEnabled(bParentEnabled) ? ESlateDrawEffect::None : ESlateDrawEffect::DisabledEffect;
+
+				DrawBackground(AllottedGeometry, OutDrawElements, BaseLayerId, DrawEffects);
+				DrawViewGrids(AllottedGeometry, MyCullingRect, OutDrawElements, BaseLayerId, DrawEffects);
+				DrawLabels(AllottedGeometry, MyCullingRect, OutDrawElements, BaseLayerId, DrawEffects);
+				DrawCurves(CurveEditor.ToSharedRef(), AllottedGeometry, MyCullingRect, OutDrawElements, BaseLayerId, InWidgetStyle, DrawEffects);
+			}
+		}
+
+		void SViewStacked::DrawLabels(const FGeometry& AllottedGeometry, const FSlateRect& MyCullingRect, FSlateWindowElementList& OutDrawElements, int32 BaseLayerId, ESlateDrawEffect DrawEffects) const
+		{
+			TSharedPtr<FCurveEditor> CurveEditor = WeakCurveEditor.Pin();
+			if (!CurveEditor)
+			{
+				return;
+			}
+
+			const int32 LabelLayerId = BaseLayerId + CurveViewConstants::ELayerOffset::Labels;
+
+			const double ValuePerPixel = 1.0 / StackedHeight;
+			const double ValueSpacePadding = StackedPadding * ValuePerPixel;
+
+			const FSlateFontInfo FontInfo = FAppStyle::GetFontStyle("CurveEd.LabelFont");
+			const FVector2D LocalSize = AllottedGeometry.GetLocalSize();
+			const FCurveEditorScreenSpace ViewSpace = GetViewSpace();
+
+			// Draw the curve labels for each view
+			for (auto It = CurveInfoByID.CreateConstIterator(); It; ++It)
+			{
+				FCurveModel* Curve = CurveEditor->FindCurve(It.Key());
+				if (!ensureAlways(Curve))
+				{
+					continue;
+				}
+
+				const int32  CurveIndexFromBottom = CurveInfoByID.Num() - It->Value.CurveIndex - 1;
+				const double PaddingToBottomOfView = (CurveIndexFromBottom + 1) * ValueSpacePadding;
+
+				const float PixelBottom = ViewSpace.ValueToScreen(CurveIndexFromBottom + PaddingToBottomOfView);
+				const float PixelTop = ViewSpace.ValueToScreen(CurveIndexFromBottom + PaddingToBottomOfView + 1.0);
+
+				if (!FSlateRect::DoRectanglesIntersect(MyCullingRect, TransformRect(AllottedGeometry.GetAccumulatedLayoutTransform(), FSlateRect(0, PixelTop, LocalSize.X, PixelBottom))))
+				{
+					continue;
+				}
+
+				const TSharedRef<FSlateFontMeasure> FontMeasure = FSlateApplication::Get().GetRenderer()->GetFontMeasureService();
+
+				// Render label
+				FText Label = Curve->GetLongDisplayName();
+
+				static const float LabelOffsetX = 15.0f;
+				static const float LabelOffsetY = 35.0f;
+				FVector2D LabelPosition(LabelOffsetX, PixelTop - LabelOffsetY);
+				FPaintGeometry LabelGeometry = AllottedGeometry.ToPaintGeometry(FSlateLayoutTransform(LabelPosition));
+				const FVector2D LabelSize = FontMeasure->Measure(Label, FontInfo);
+
+				FSlateDrawElement::MakeText(OutDrawElements, LabelLayerId + 1, LabelGeometry, Label, FontInfo, DrawEffects, Curve->GetColor());
+
+				// Render axes descriptor
+				FText Descriptor = static_cast<FWaveTableCurveModelBase*>(Curve)->GetAxesDescriptor();
+
+				const FVector2D DescriptorSize = FontMeasure->Measure(Descriptor, FontInfo);
+
+				static const float LabelBufferX = 20.0f; // Keeps label and axes descriptor visually separated
+				static const float GutterBufferX = 20.0f; // Accounts for potential scroll bar
+				const float ViewWidth = ViewSpace.GetPhysicalWidth();
+				const float FloatingDescriptorX = ViewWidth - DescriptorSize.X;
+				LabelPosition = FVector2D(FMath::Max(LabelSize.X + LabelBufferX, FloatingDescriptorX - GutterBufferX), PixelTop - LabelOffsetY);
+				LabelGeometry = AllottedGeometry.ToPaintGeometry(FSlateLayoutTransform(LabelPosition));
+
+				FSlateDrawElement::MakeText(OutDrawElements, LabelLayerId + 1, LabelGeometry, Descriptor, FontInfo, DrawEffects, Curve->GetColor());
+			}
+		}
+
+		void SViewStacked::DrawViewGrids(const FGeometry& AllottedGeometry, const FSlateRect& MyCullingRect, FSlateWindowElementList& OutDrawElements, int32 BaseLayerId, ESlateDrawEffect DrawEffects) const
+		{
+			TSharedPtr<FCurveEditor> CurveEditor = WeakCurveEditor.Pin();
+			if (!CurveEditor)
+			{
+				return;
+			}
+
+			// Rendering info
+			const float          Width = AllottedGeometry.GetLocalSize().X;
+			const float          Height = AllottedGeometry.GetLocalSize().Y;
+			const FSlateBrush*   WhiteBrush = FAppStyle::GetBrush("WhiteBrush");
+
+			FGridDrawInfo DrawInfo(&AllottedGeometry, GetViewSpace(), CurveEditor->GetPanel()->GetGridLineTint(), BaseLayerId);
+
+			TArray<float> MajorGridLinesX, MinorGridLinesX;
+			TArray<FText> MajorGridLabelsX;
+
+			GetGridLinesX(CurveEditor.ToSharedRef(), MajorGridLinesX, MinorGridLinesX, &MajorGridLabelsX);
+
+			const double ValuePerPixel = 1.0 / StackedHeight;
+			const double ValueSpacePadding = StackedPadding * ValuePerPixel;
+
+			for (auto It = CurveInfoByID.CreateConstIterator(); It; ++It)
+			{
+				DrawInfo.SetCurveModel(CurveEditor->FindCurve(It.Key()));
+				if (!ensureAlways(DrawInfo.GetCurveModel()))
+				{
+					continue;
+				}
+
+				TArray<FText> CurveModelGridLabelsX = MajorGridLabelsX;
+				check(MajorGridLinesX.Num() == CurveModelGridLabelsX.Num());
+
+				if (const FWaveTableCurveModelBase* EditorModel = static_cast<const FWaveTableCurveModelBase*>(DrawInfo.GetCurveModel()))
+				{
+					for (int32 i = 0; i < CurveModelGridLabelsX.Num(); ++i)
+					{
+						FText& InputLabel = CurveModelGridLabelsX[i];
+						FormatInputLabel(*EditorModel, DrawInfo.LabelFormat, InputLabel);
+					}
+				}
+
+				const int32 Index = CurveInfoByID.Num() - It->Value.CurveIndex - 1;
+				double Padding = (Index + 1) * ValueSpacePadding;
+				DrawInfo.SetLowerValue(Index + Padding);
+
+				FSlateRect BoundsRect(0, DrawInfo.GetPixelTop(), Width, DrawInfo.GetPixelBottom());
+				if (!FSlateRect::DoRectanglesIntersect(MyCullingRect, TransformRect(AllottedGeometry.GetAccumulatedLayoutTransform(), BoundsRect)))
+				{
+					continue;
+				}
+
+				// Tint the views based on their curve color
+				{
+					FLinearColor CurveColorTint = DrawInfo.GetCurveModel()->GetColor().CopyWithNewOpacity(0.05f);
+					const FPaintGeometry BoxGeometry = AllottedGeometry.ToPaintGeometry(
+						FVector2D(Width, StackedHeight),
+						FSlateLayoutTransform(FVector2D(0.f, DrawInfo.GetPixelTop()))
+					);
+
+					const int32 GridOverlayLayerId = DrawInfo.GetBaseLayerId() + CurveViewConstants::ELayerOffset::GridOverlays;
+					FSlateDrawElement::MakeBox(OutDrawElements, GridOverlayLayerId, BoxGeometry, WhiteBrush, DrawEffects, CurveColorTint);
+				}
+
+				// Horizontal grid lines
+				DrawInfo.LinePoints[0].X = 0.0;
+				DrawInfo.LinePoints[1].X = Width;
+
+				DrawViewGridLineX(OutDrawElements, DrawInfo, DrawEffects, 1.0  /* OffsetAlpha */, true  /* bIsMajor */);
+				DrawViewGridLineX(OutDrawElements, DrawInfo, DrawEffects, 0.75 /* OffsetAlpha */, false /* bIsMajor */);
+				DrawViewGridLineX(OutDrawElements, DrawInfo, DrawEffects, 0.5  /* OffsetAlpha */, true  /* bIsMajor */);
+				DrawViewGridLineX(OutDrawElements, DrawInfo, DrawEffects, 0.25 /* OffsetAlpha */, false /* bIsMajor */);
+				DrawViewGridLineX(OutDrawElements, DrawInfo, DrawEffects, 0.0  /* OffsetAlpha */, true  /* bIsMajor */);
+
+				// Vertical grid lines
+				{
+					DrawInfo.LinePoints[0].Y = DrawInfo.GetPixelTop();
+					DrawInfo.LinePoints[1].Y = DrawInfo.GetPixelBottom();
+
+					// Draw major vertical grid lines
+					for (int i = 0; i < MajorGridLinesX.Num(); ++i)
+					{
+						const float VerticalLine = FMath::RoundToFloat(MajorGridLinesX[i]);
+						const FText& Label = CurveModelGridLabelsX[i];
+
+						DrawViewGridLineY(VerticalLine, OutDrawElements, DrawInfo, DrawEffects, &Label, true /* bIsMajor */);
+					}
+
+					// Now draw the minor vertical lines which are drawn with a lighter color.
+					for (float VerticalLine : MinorGridLinesX)
+					{
+						VerticalLine = FMath::RoundToFloat(VerticalLine);
+						DrawViewGridLineY(VerticalLine, OutDrawElements, DrawInfo, DrawEffects, nullptr /* Label */, false /* bIsMajor */);
+					}
+				}
+			}
+		}
+
+		void SViewStacked::DrawViewGridLineX(FSlateWindowElementList& OutDrawElements, FGridDrawInfo& DrawInfo, ESlateDrawEffect DrawEffects, double OffsetAlpha, bool bIsMajor) const
+		{
+			double ValueMin;
+			double ValueMax;
+			DrawInfo.GetCurveModel()->GetValueRange(ValueMin, ValueMax);
+
+			const int32 GridLineLayerId = DrawInfo.GetBaseLayerId() + CurveViewConstants::ELayerOffset::GridLines;
+
+			const double LowerValue = DrawInfo.GetLowerValue();
+			const double PixelBottom = DrawInfo.GetPixelBottom();
+			const double PixelTop = DrawInfo.GetPixelTop();
+
+			const TSharedRef<FSlateFontMeasure> FontMeasure = FSlateApplication::Get().GetRenderer()->GetFontMeasureService();
+			const FSlateFontInfo FontInfo = FCoreStyle::Get().GetFontStyle("CurveEd.InfoFont");
+
+			FLinearColor Color = bIsMajor ? DrawInfo.GetMajorGridColor() : DrawInfo.GetMinorGridColor();
+
+			DrawInfo.LinePoints[0].Y = DrawInfo.LinePoints[1].Y = DrawInfo.ScreenSpace.ValueToScreen(LowerValue + OffsetAlpha);
+			FSlateDrawElement::MakeLines(OutDrawElements, GridLineLayerId, DrawInfo.PaintGeometry, DrawInfo.LinePoints, DrawEffects, Color, false);
+
+			FText Label = FText::AsNumber(FMath::Lerp(ValueMin, ValueMax, OffsetAlpha), &DrawInfo.LabelFormat);
+
+			if (const FWaveTableCurveModelBase* CurveModel = static_cast<const FWaveTableCurveModelBase*>(DrawInfo.GetCurveModel()))
+			{
+				FormatOutputLabel(*CurveModel, DrawInfo.LabelFormat, Label);
+			}
+
+			const FVector2D LabelSize = FontMeasure->Measure(Label, FontInfo);
+			double LabelY = FMath::Lerp(PixelBottom, PixelTop, OffsetAlpha);
+	
+			// Offset label Y position below line only if the top gridline,
+			// otherwise push above
+			if (FMath::IsNearlyEqual(OffsetAlpha, 1.0))
+			{
+				LabelY += 5.0;
+			}
+			else
+			{
+				LabelY -= 15.0;
+			}
+
+			const FPaintGeometry LabelGeometry = DrawInfo.AllottedGeometry->ToPaintGeometry(FSlateLayoutTransform(FVector2D(CurveViewConstants::CurveLabelOffsetX, LabelY)));
+
+			FSlateDrawElement::MakeText(
+				OutDrawElements,
+				DrawInfo.GetBaseLayerId() + CurveViewConstants::ELayerOffset::GridLabels,
+				LabelGeometry,
+				Label,
+				FontInfo,
+				DrawEffects,
+				DrawInfo.GetLabelColor()
+			);
+		}
+
+		void SViewStacked::DrawViewGridLineY(const float VerticalLine, FSlateWindowElementList& OutDrawElements, FGridDrawInfo& DrawInfo, ESlateDrawEffect DrawEffects, const FText* Label, bool bIsMajor) const
+		{
+			const float Width = DrawInfo.AllottedGeometry->GetLocalSize().X;
+			if (VerticalLine >= 0 || VerticalLine <= FMath::RoundToFloat(Width))
+			{
+				const int32 GridLineLayerId = DrawInfo.GetBaseLayerId() + CurveViewConstants::ELayerOffset::GridLines;
+
+				const FLinearColor LineColor = bIsMajor ? DrawInfo.GetMajorGridColor() : DrawInfo.GetMinorGridColor();
+
+				DrawInfo.LinePoints[0].X = DrawInfo.LinePoints[1].X = VerticalLine;
+				FSlateDrawElement::MakeLines(OutDrawElements, GridLineLayerId, DrawInfo.PaintGeometry, DrawInfo.LinePoints, DrawEffects, LineColor, false);
+
+				if (Label)
+				{
+					const FSlateFontInfo FontInfo = FCoreStyle::Get().GetFontStyle("CurveEd.InfoFont");
+
+					const TSharedRef<FSlateFontMeasure> FontMeasure = FSlateApplication::Get().GetRenderer()->GetFontMeasureService();
+					const FVector2D LabelSize = FontMeasure->Measure(*Label, FontInfo);
+					const FPaintGeometry LabelGeometry = DrawInfo.AllottedGeometry->ToPaintGeometry(FSlateLayoutTransform(FVector2D(VerticalLine, DrawInfo.LinePoints[0].Y - LabelSize.Y * 1.2f)));
+
+					FSlateDrawElement::MakeText(
+						OutDrawElements,
+						DrawInfo.GetBaseLayerId() + CurveViewConstants::ELayerOffset::GridLabels,
+						LabelGeometry,
+						*Label,
+						FontInfo,
+						DrawEffects,
+						DrawInfo.GetLabelColor()
+					);
+				}
+			}
+		}
+
+		SViewStacked::FGridDrawInfo::FGridDrawInfo(const FGeometry* InAllottedGeometry, const FCurveEditorScreenSpace& InScreenSpace, FLinearColor InGridColor, int32 InBaseLayerId)
+			: AllottedGeometry(InAllottedGeometry)
+			, ScreenSpace(InScreenSpace)
+			, BaseLayerId(InBaseLayerId)
+			, CurveModel(nullptr)
+			, LowerValue(0)
+			, PixelBottom(0)
+			, PixelTop(0)
+		{
+			check(AllottedGeometry);
+
+			// Pre-allocate an array of line points to draw our vertical lines. Each major grid line
+			// will overwrite the X value of both points but leave the Y value untouched so they draw from the bottom to the top.
+			LinePoints.Add(FVector2D(0.f, 0.f));
+			LinePoints.Add(FVector2D(0.f, 0.f));
+
+			MajorGridColor = InGridColor;
+			MinorGridColor = InGridColor.CopyWithNewOpacity(InGridColor.A * 0.5f);
+
+			PaintGeometry = InAllottedGeometry->ToPaintGeometry();
+
+			LabelFormat.SetMaximumFractionalDigits(2);
+		}
+
+		void SViewStacked::FGridDrawInfo::SetCurveModel(const FCurveModel* InCurveModel)
+		{
+			CurveModel = InCurveModel;
+		}
+
+		const FCurveModel* SViewStacked::FGridDrawInfo::GetCurveModel() const
+		{
+			return CurveModel;
+		}
+
+		void SViewStacked::FGridDrawInfo::SetLowerValue(double InLowerValue)
+		{
+			LowerValue = InLowerValue;
+			PixelBottom = ScreenSpace.ValueToScreen(InLowerValue);
+			PixelTop = ScreenSpace.ValueToScreen(InLowerValue + 1.0);
+		}
+
+		int32 SViewStacked::FGridDrawInfo::GetBaseLayerId() const
+		{
+			return BaseLayerId;
+		}
+
+		FLinearColor SViewStacked::FGridDrawInfo::GetLabelColor() const
+		{
+			check(CurveModel);
+			return GetCurveModel()->GetColor().CopyWithNewOpacity(0.7f);
+		}
+
+		double SViewStacked::FGridDrawInfo::GetLowerValue() const
+		{
+			return LowerValue;
+		}
+
+		FLinearColor SViewStacked::FGridDrawInfo::GetMajorGridColor() const
+		{
+			return MajorGridColor;
+		}
+
+		FLinearColor SViewStacked::FGridDrawInfo::GetMinorGridColor() const
+		{
+			return MinorGridColor;
+		}
+
+		double SViewStacked::FGridDrawInfo::GetPixelBottom() const
+		{
+			return PixelBottom;
+		}
+
+		double SViewStacked::FGridDrawInfo::GetPixelTop() const
+		{
+			return PixelTop;
+		}
+	} // namespace Editor
+} // namespace WaveTable
+
+#undef LOCTEXT_NAMESPACE
\ No newline at end of file
diff --git a/Engine/Plugins/Runtime/WaveTable/Source/WaveTableEditor/Private/WaveTableEditorModule.cpp b/Engine/Plugins/Runtime/WaveTable/Source/WaveTableEditor/Private/WaveTableEditorModule.cpp
new file mode 100644
--- /dev/null
+++ b/Engine/Plugins/Runtime/WaveTable/Source/WaveTableEditor/Private/WaveTableEditorModule.cpp
@@ -0,0 +1,21 @@
+// Copyright Epic Games, Inc. All Rights Reserved.
+#include "WaveTableEditorModule.h"
+
+
+DEFINE_LOG_CATEGORY(LogWaveTableEditor);
+
+namespace WaveTable
+{
+	namespace Editor
+	{
+		void FModule::StartupModule()
+		{
+		}
+
+		void FModule::ShutdownModule()
+		{
+		}
+	} // namespace Editor
+} // namespace WaveTable
+
+IMPLEMENT_MODULE(WaveTable::Editor::FModule, WaveTableEditor);
diff --git a/Engine/Plugins/Runtime/WaveTable/Source/WaveTableEditor/Private/WaveTableEditorModule.h b/Engine/Plugins/Runtime/WaveTable/Source/WaveTableEditor/Private/WaveTableEditorModule.h
new file mode 100644
--- /dev/null
+++ b/Engine/Plugins/Runtime/WaveTable/Source/WaveTableEditor/Private/WaveTableEditorModule.h
@@ -0,0 +1,24 @@
+// Copyright Epic Games, Inc. All Rights Reserved.
+#pragma once
+
+#include "Logging/LogMacros.h"
+#include "Modules/ModuleManager.h"
+
+
+DECLARE_LOG_CATEGORY_EXTERN(LogWaveTableEditor, Log, All);
+
+namespace WaveTable
+{
+	namespace Editor
+	{
+		class FModule : public IModuleInterface
+		{
+		public:
+			FModule() = default;
+
+			virtual void StartupModule() override;
+
+			virtual void ShutdownModule() override;
+		};
+	} // namespace Editor
+} // namespace WaveTable
diff --git a/Engine/Plugins/Runtime/WaveTable/Source/WaveTableEditor/Private/WaveTableFileUtilities.cpp b/Engine/Plugins/Runtime/WaveTable/Source/WaveTableEditor/Private/WaveTableFileUtilities.cpp
new file mode 100644
--- /dev/null
+++ b/Engine/Plugins/Runtime/WaveTable/Source/WaveTableEditor/Private/WaveTableFileUtilities.cpp
@@ -0,0 +1,60 @@
+// Copyright Epic Games, Inc. All Rights Reserved.
+#include "WaveTableFileUtilities.h"
+
+#include "Factories/SoundFactory.h"
+#include "HAL/Platform.h"
+#include "Sound/SoundWave.h"
+#include "UObject/SoftObjectPath.h"
+#include "UObject/Object.h"
+#include "UObject/ObjectMacros.h"
+#include "UObject/Package.h"
+#include "Utils.h"
+#include "WaveTableSampler.h"
+
+
+namespace WaveTable
+{
+	namespace Editor
+	{
+		namespace FileUtilities
+		{
+			void LoadPCMChannel(const FString& InFilePath, int32 InChannelIndex, TArray<float>& OutPCMData)
+			{
+				OutPCMData.Empty();
+
+				if (!InFilePath.IsEmpty())
+				{
+					if (USoundFactory* SoundWaveFactory = NewObject<USoundFactory>())
+					{
+						// Setup sane defaults for importing localized sound waves
+						SoundWaveFactory->bAutoCreateCue = false;
+						SoundWaveFactory->SuppressImportDialogs();
+
+						UPackage* TempPackage = CreatePackage(TEXT("/Temp/WaveTable"));
+						if (USoundWave* SoundWave = ImportObject<USoundWave>(TempPackage, "ImportedWaveTable", RF_Public | RF_Standalone, *InFilePath, nullptr, SoundWaveFactory))
+						{
+							TArray<uint8> RawPCMData;
+							uint32 SampleRate = 0;
+							uint16 NumChannels = 0;
+							SoundWave->GetImportedSoundWaveData(RawPCMData, SampleRate, NumChannels);
+
+							if (NumChannels > 0)
+							{
+								const int32 NumFrames = RawPCMData.Num() / NumChannels / sizeof(int16);
+								OutPCMData.Empty();
+								OutPCMData.AddZeroed(NumFrames);
+								const int16* RawDataPtr = (const int16*)(RawPCMData.GetData());
+								const int32 ChannelOffset = (InChannelIndex % NumChannels) * sizeof(int16);
+								constexpr float Max16BitAsFloat = static_cast<float>(TNumericLimits<int16>::Max());
+								for (int32 i = 0; i < OutPCMData.Num(); ++i)
+								{
+									OutPCMData[i] = (RawDataPtr[(i * NumChannels) + ChannelOffset]) / Max16BitAsFloat;
+								}
+							}
+						}
+					}
+				}
+			}
+		} // namespace FileUtilities
+	} // namespace Editor
+} // namespace WaveTable
diff --git a/Engine/Plugins/Runtime/WaveTable/Source/WaveTableEditor/Public/WaveTableBankEditor.h b/Engine/Plugins/Runtime/WaveTable/Source/WaveTableEditor/Public/WaveTableBankEditor.h
new file mode 100644
--- /dev/null
+++ b/Engine/Plugins/Runtime/WaveTable/Source/WaveTableEditor/Public/WaveTableBankEditor.h
@@ -0,0 +1,116 @@
+// Copyright Epic Games, Inc. All Rights Reserved.
+
+#pragma once
+
+#include "CurveEditorTypes.h"
+#include "Curves/CurveOwnerInterface.h"
+#include "EditorUndoClient.h"
+#include "Framework/Docking/TabManager.h"
+#include "Math/Color.h"
+#include "Misc/NotifyHook.h"
+#include "WaveTableCurveEditorViewStacked.h"
+#include "Toolkits/AssetEditorToolkit.h"
+#include "Toolkits/IToolkitHost.h"
+#include "WaveTableSettings.h"
+#include "Widgets/SWidget.h"
+
+
+// Forward Declarations
+class FCurveEditor;
+class IToolkitHost;
+class SCurveEditorPanel;
+class UCurveFloat;
+
+namespace WaveTable
+{
+	namespace Editor
+	{
+		class WAVETABLEEDITOR_API FWaveTableBankEditor : public FAssetEditorToolkit, public FNotifyHook, public FEditorUndoClient
+		{
+		public:
+			FWaveTableBankEditor();
+			virtual ~FWaveTableBankEditor() = default;
+
+			void Init(const EToolkitMode::Type Mode, const TSharedPtr<IToolkitHost>& InitToolkitHost, UObject* InParentObject);
+
+			/** FAssetEditorToolkit interface */
+			virtual FName GetToolkitFName() const override;
+			virtual FText GetBaseToolkitName() const override;
+			virtual FString GetWorldCentricTabPrefix() const override;
+			virtual FLinearColor GetWorldCentricTabColorScale() const override;
+			virtual void RegisterTabSpawners(const TSharedRef<FTabManager>& TabManager) override;
+			virtual void UnregisterTabSpawners(const TSharedRef<FTabManager>& TabManager) override;
+
+			/** FNotifyHook interface */
+			virtual void NotifyPostChange(const FPropertyChangedEvent& PropertyChangedEvent, FProperty* PropertyThatChanged) override;
+
+		protected:
+			struct FCurveData
+			{
+				FCurveModelID ModelID;
+				TSharedPtr<FRichCurve> ExpressionCurve;
+
+				FCurveData()
+					: ModelID(FCurveModelID::Unique())
+				{
+				}
+			};
+
+			virtual void PostUndo(bool bSuccess) override;
+			virtual void PostRedo(bool bSuccess) override;
+
+			virtual bool GetIsPropertyEditorDisabled() const = 0;
+
+			/** Generates expression curve at the given index. */
+			virtual void GenerateExpressionCurve(FCurveData& OutCurveData, int32 InInputIndex, EWaveTableCurveSource InSource, bool bInIsUnset = false) = 0;
+
+			virtual TUniquePtr<FWaveTableCurveModelBase> ConstructCurveModel(FRichCurve& InRichCurve, UObject* InParentObject, EWaveTableCurveSource InSource) = 0;
+
+			virtual EWaveTableCurve GetCurveType(int32 InInputIndex) const = 0;
+			virtual FRichCurve& GetCustomCurveChecked(int32 InInputIndex) const = 0;
+			virtual UCurveFloat* GetSharedCurve(int32 InInputIndex) const = 0;
+			virtual int32 GetNumCurves() const = 0;
+
+			void SetCurve(int32 InInputIndex, FRichCurve& InRichCurve, EWaveTableCurveSource InSource);
+
+		private:
+			void ResetCurves();
+			void InitCurves();
+
+			/** Updates & redraws curves. */
+			void RefreshCurves();
+
+			/**	Spawns the tab allowing for editing/viewing the output curve(s) */
+			TSharedRef<SDockTab> SpawnTab_OutputCurve(const FSpawnTabArgs& Args);
+
+			/**	Spawns the tab allowing for editing/viewing the output curve(s) */
+			TSharedRef<SDockTab> SpawnTab_Properties(const FSpawnTabArgs& Args);
+
+			/** Get the orientation for the snap value controls. */
+			EOrientation GetSnapLabelOrientation() const;
+
+			/** Trims keys out-of-bounds in provided curve */
+			static void TrimKeys(FRichCurve& OutCurve);
+
+			/** Clears the expression curve at the given input index */
+			void ClearExpressionCurve(int32 InInputIndex);
+
+			bool RequiresNewCurve(int32 InInputIndex, const FRichCurve& InRichCurve) const;
+
+			TSharedPtr<FUICommandList> ToolbarCurveTargetCommands;
+
+			TSharedPtr<FCurveEditor> CurveEditor;
+			TSharedPtr<SCurveEditorPanel> CurvePanel;
+
+			TArray<FCurveData> CurveData;
+
+			/** Properties tab */
+			TSharedPtr<IDetailsView> PropertiesView;
+
+			/** Settings Editor App Identifier */
+			static const FName AppIdentifier;
+			static const FName CurveTabId;
+			static const FName PropertiesTabId;
+		};
+	} // namespace Editor
+} // namespace WaveTable
\ No newline at end of file
diff --git a/Engine/Plugins/Runtime/WaveTable/Source/WaveTableEditor/Public/WaveTableCurveEditorViewStacked.h b/Engine/Plugins/Runtime/WaveTable/Source/WaveTableEditor/Public/WaveTableCurveEditorViewStacked.h
new file mode 100644
--- /dev/null
+++ b/Engine/Plugins/Runtime/WaveTable/Source/WaveTableEditor/Public/WaveTableCurveEditorViewStacked.h
@@ -0,0 +1,139 @@
+// Copyright Epic Games, Inc. All Rights Reserved.
+#pragma once
+
+#include "CurveModel.h"
+#include "Curves/RichCurve.h"
+#include "Internationalization/Text.h"
+#include "RichCurveEditorModel.h"
+#include "SGraphActionMenu.h"
+#include "UObject/Object.h"
+#include "UObject/WeakObjectPtr.h"
+#include "UObject/WeakObjectPtrTemplates.h"
+#include "Views/SCurveEditorViewStacked.h"
+#include "Views/SInteractiveCurveEditorView.h"
+
+
+// Forward Declarations
+class UCurveFloat;
+
+
+UENUM()
+enum class EWaveTableCurveSource : uint8
+{
+	Custom,
+	Expression,
+	Shared,
+	Unset
+};
+
+
+namespace WaveTable
+{
+	namespace Editor
+	{
+		class WAVETABLEEDITOR_API FWaveTableCurveModelBase : public FRichCurveEditorModelRaw
+		{
+		public:
+			static ECurveEditorViewID ViewId;
+
+			FWaveTableCurveModelBase(FRichCurve& InRichCurve, UObject* InOwner, EWaveTableCurveSource InSource);
+
+			const FText& GetAxesDescriptor() const;
+			const UObject* GetParentObject() const;
+			EWaveTableCurveSource GetSource() const;
+			void Refresh(int32 InCurveIndex);
+
+			virtual bool IsReadOnly() const override;
+			virtual FLinearColor GetColor() const override;
+
+			int32 GetCurveIndex() const
+			{
+				return CurveIndex;
+			}
+
+		protected:
+			virtual void RefreshCurveDescriptorText(FText& OutShortDisplayName, FText& OutInputAxisName, FText& OutOutputAxisName) = 0;
+
+			virtual UCurveFloat* GetSharedCurve() = 0;
+
+			virtual FColor GetCurveColor() const = 0;
+			virtual bool GetPropertyEditorDisabled() const = 0;
+			virtual FText GetPropertyEditorDisabledText() const = 0;
+
+			TWeakObjectPtr<UObject> ParentObject;
+
+		private:
+			int32 CurveIndex = INDEX_NONE;
+			EWaveTableCurveSource Source = EWaveTableCurveSource::Unset;
+
+			FText InputAxisName;
+			FText AxesDescriptor;
+		};
+
+		class WAVETABLEEDITOR_API SViewStacked : public SCurveEditorViewStacked
+		{
+		public:
+			void Construct(const FArguments& InArgs, TWeakPtr<FCurveEditor> InCurveEditor);
+
+		protected:
+			virtual void PaintView(
+				const FPaintArgs& Args,
+				const FGeometry& AllottedGeometry,
+				const FSlateRect& MyCullingRect,
+				FSlateWindowElementList& OutDrawElements,
+				int32 BaseLayerId, 
+				const FWidgetStyle& InWidgetStyle,
+				bool bParentEnabled) const override;
+
+			virtual void DrawViewGrids(const FGeometry& AllottedGeometry, const FSlateRect& MyCullingRect, FSlateWindowElementList& OutDrawElements, int32 BaseLayerId, ESlateDrawEffect DrawEffects) const override;
+			virtual void DrawLabels(const FGeometry& AllottedGeometry, const FSlateRect& MyCullingRect, FSlateWindowElementList& OutDrawElements, int32 BaseLayerId, ESlateDrawEffect DrawEffects) const override;
+
+			virtual void FormatInputLabel(const FWaveTableCurveModelBase& EditorModel, const FNumberFormattingOptions& InLabelFormat, FText& InOutLabel) const { }
+			virtual void FormatOutputLabel(const FWaveTableCurveModelBase& EditorModel, const FNumberFormattingOptions& InLabelFormat, FText& InOutLabel) const { }
+
+		private:
+			struct FGridDrawInfo
+			{
+				const FGeometry* AllottedGeometry;
+
+				FPaintGeometry PaintGeometry;
+
+				ESlateDrawEffect DrawEffects;
+				FNumberFormattingOptions LabelFormat;
+
+				TArray<FVector2D> LinePoints;
+
+				FCurveEditorScreenSpace ScreenSpace;
+
+			private:
+				FLinearColor MajorGridColor;
+				FLinearColor MinorGridColor;
+
+				int32 BaseLayerId = INDEX_NONE;
+
+				const FCurveModel* CurveModel = nullptr;
+
+				double LowerValue = 0.0;
+				double PixelBottom = 0.0;
+				double PixelTop = 0.0;
+
+			public:
+				FGridDrawInfo(const FGeometry* InAllottedGeometry, const FCurveEditorScreenSpace& InScreenSpace, FLinearColor InGridColor, int32 InBaseLayerId);
+
+				void SetCurveModel(const FCurveModel* InCurveModel);
+				const FCurveModel* GetCurveModel() const;
+				void SetLowerValue(double InLowerValue);
+				int32 GetBaseLayerId() const;
+				FLinearColor GetLabelColor() const;
+				double GetLowerValue() const;
+				FLinearColor GetMajorGridColor() const;
+				FLinearColor GetMinorGridColor() const;
+				double GetPixelBottom() const;
+				double GetPixelTop() const;
+			};
+
+			void DrawViewGridLineX(FSlateWindowElementList& OutDrawElements, FGridDrawInfo& DrawInfo, ESlateDrawEffect DrawEffect, double OffsetAlpha, bool bIsMajor) const;
+			void DrawViewGridLineY(const float VerticalLine, FSlateWindowElementList& OutDrawElements, FGridDrawInfo &DrawInfo, ESlateDrawEffect DrawEffects, const FText* Label, bool bIsMajor) const;
+		};
+	} // namespace Editor
+} // namespace WaveTable
diff --git a/Engine/Plugins/Runtime/WaveTable/Source/WaveTableEditor/Public/WaveTableFileUtilities.h b/Engine/Plugins/Runtime/WaveTable/Source/WaveTableEditor/Public/WaveTableFileUtilities.h
new file mode 100644
--- /dev/null
+++ b/Engine/Plugins/Runtime/WaveTable/Source/WaveTableEditor/Public/WaveTableFileUtilities.h
@@ -0,0 +1,17 @@
+// Copyright Epic Games, Inc. All Rights Reserved.
+#pragma once
+
+#include "Containers/Array.h"
+#include "Containers/UnrealString.h"
+
+
+namespace WaveTable
+{
+	namespace Editor
+	{
+		namespace FileUtilities
+		{
+			void WAVETABLEEDITOR_API LoadPCMChannel(const FString& InFilePath, int32 InChannelIndex, TArray<float>& OutPCMData);
+		} // namespace FileUtilities
+	} // namespace Editor
+} // namespace WaveTable
diff --git a/Engine/Plugins/Runtime/WaveTable/Source/WaveTableEditor/WaveTableEditor.Build.cs b/Engine/Plugins/Runtime/WaveTable/Source/WaveTableEditor/WaveTableEditor.Build.cs
new file mode 100644
--- /dev/null
+++ b/Engine/Plugins/Runtime/WaveTable/Source/WaveTableEditor/WaveTableEditor.Build.cs
@@ -0,0 +1,32 @@
+// Copyright Epic Games, Inc. All Rights Reserved.
+
+using UnrealBuildTool;
+
+public class WaveTableEditor : ModuleRules
+{
+	public WaveTableEditor(ReadOnlyTargetRules Target) : base(Target)
+	{
+		PrivateDependencyModuleNames.AddRange(
+			new string[]
+			{
+				"AudioEditor",
+				"AudioExtensions",
+				"Core",
+				"CoreUObject",
+				"CurveEditor",
+				"Engine",
+				"EditorFramework",
+				"EditorWidgets",
+				"GameProjectGeneration",
+				"InputCore",
+				"PropertyEditor",
+				"SequenceRecorder",
+				"Slate",
+				"SlateCore",
+				"ToolWidgets",
+				"UnrealEd",
+				"WaveTable",
+			}
+		);
+	}
+}
diff --git a/Engine/Plugins/Runtime/WaveTable/WaveTable.uplugin b/Engine/Plugins/Runtime/WaveTable/WaveTable.uplugin
new file mode 100644
--- /dev/null
+++ b/Engine/Plugins/Runtime/WaveTable/WaveTable.uplugin
@@ -0,0 +1,32 @@
+{
+	"FileVersion": 3,
+	"Version": 1,
+	"VersionName": "1.0",
+	"FriendlyName": "Wave Tables",
+	"Description": "Default implementation of WaveTable support within the Unreal Audio Engine.",
+	"Category": "Audio",
+	"CreatedBy": "Epic Games, Inc.",
+	"CreatedByURL": "http://epicgames.com",
+	"DocsURL": "",
+	"MarketplaceURL": "",
+	"SupportURL": "",
+	"EnabledByDefault": false,
+	"CanContainContent": true,
+	"IsBetaVersion": true,
+	"IsExperimentalVersion": false,
+	"Installed": false,
+	"Modules": [
+		{
+			"Name": "WaveTable",
+			"Type": "Runtime",
+			"LoadingPhase": "Default"
+		},
+		{
+			"Name": "WaveTableEditor",
+			"Type": "Editor",
+			"LoadingPhase": "PostEngineInit"
+		}
+	],
+	"Plugins": [
+	]	
+}
\ No newline at end of file
