Remove optimize never
#rnx

diff --git a/Engine/Plugins/Runtime/WaveTable/Source/WaveTable/WaveTable.Build.cs b/Engine/Plugins/Runtime/WaveTable/Source/WaveTable/WaveTable.Build.cs
--- a/Engine/Plugins/Runtime/WaveTable/Source/WaveTable/WaveTable.Build.cs
+++ b/Engine/Plugins/Runtime/WaveTable/Source/WaveTable/WaveTable.Build.cs
@@ -6,8 +6,6 @@ namespace UnrealBuildTool.Rules
 	{
 		public WaveTable(ReadOnlyTargetRules Target) : base(Target)
 		{
-			OptimizeCode = CodeOptimization.Never;
-
 			PublicDependencyModuleNames.AddRange(
 				new string[]
 				{
