WaveTable Checkpoint:
Envelope, getter, & oscillator nodes (Enablement will come in future CL)
- WaveTable Clean-up, full file res, & interactive edit support
- Add ability to use full sample size in WaveTableBank
- Remove ability to use WTs for ModulationPatches.  Not a good enough use case and enveloping UX still not at at a usable place.
- Add WaveTableData size & length [REDACTED_HANDLE] properties on WaveTableBanks for eyeballing memory/ usage & general size (editor only props)
#rnx

diff --git a/Engine/Plugins/Runtime/AudioModulation/Source/AudioModulation/AudioModulation.Build.cs b/Engine/Plugins/Runtime/AudioModulation/Source/AudioModulation/AudioModulation.Build.cs
--- a/Engine/Plugins/Runtime/AudioModulation/Source/AudioModulation/AudioModulation.Build.cs
+++ b/Engine/Plugins/Runtime/AudioModulation/Source/AudioModulation/AudioModulation.Build.cs
@@ -16,7 +16,8 @@ namespace UnrealBuildTool.Rules
 				new string[]
 				{
 					"AudioExtensions",
-					"Core"
+					"Core",
+					"WaveTable"
 				}
 			);
 
@@ -27,8 +28,7 @@ namespace UnrealBuildTool.Rules
 					"CoreUObject",
 					"DeveloperSettings",
 					"Engine",
-					"SignalProcessing",
-					"WaveTable"
+					"SignalProcessing"
 				}
 			);
 
diff --git a/Engine/Plugins/Runtime/AudioModulation/Source/AudioModulation/Private/SoundModulationPatch.cpp b/Engine/Plugins/Runtime/AudioModulation/Source/AudioModulation/Private/SoundModulationPatch.cpp
--- a/Engine/Plugins/Runtime/AudioModulation/Source/AudioModulation/Private/SoundModulationPatch.cpp
+++ b/Engine/Plugins/Runtime/AudioModulation/Source/AudioModulation/Private/SoundModulationPatch.cpp
@@ -40,14 +40,6 @@ const Audio::FModulationParameter& USoundModulationPatch::GetOutputParameter() c
 #if WITH_EDITOR
 void USoundModulationPatch::PostEditChangeProperty(FPropertyChangedEvent& InPropertyChangedEvent)
 {
-	if (InPropertyChangedEvent.Property && InPropertyChangedEvent.ChangeType != EPropertyChangeType::Interactive)
-	{
-		for (FSoundControlModulationInput& Input : PatchSettings.Inputs)
-		{
-			Input.Transform.CacheWaveTable(PatchSettings.WaveTableResolution, false /* bBipolar */);
-		}
-	}
-
 	AudioModulation::IterateModulationManagers([this](AudioModulation::FAudioModulationManager& OutModulation)
 	{
 		OutModulation.UpdateModulator(*this);
@@ -58,14 +50,6 @@ void USoundModulationPatch::PostEditChangeProperty(FPropertyChangedEvent& InProp
 
 void USoundModulationPatch::PostEditChangeChainProperty(FPropertyChangedChainEvent& InPropertyChangedEvent)
 {
-	if (InPropertyChangedEvent.Property && InPropertyChangedEvent.ChangeType != EPropertyChangeType::Interactive)
-	{
-		for (FSoundControlModulationInput& Input : PatchSettings.Inputs)
-		{
-			Input.Transform.CacheWaveTable(PatchSettings.WaveTableResolution, false /* bBipolar */);
-		}
-	}
-
 	AudioModulation::IterateModulationManagers([this](AudioModulation::FAudioModulationManager& OutModulation)
 	{
 		OutModulation.UpdateModulator(*this);
@@ -76,14 +60,6 @@ void USoundModulationPatch::PostEditChangeChainProperty(FPropertyChangedChainEve
 
 void USoundModulationPatch::PreSave(FObjectPreSaveContext InSaveContext)
 {
-	if (!InSaveContext.IsCooking())
-	{
-		for (FSoundControlModulationInput& Input : PatchSettings.Inputs)
-		{
-			Input.Transform.CacheWaveTable(PatchSettings.WaveTableResolution, false /* bBipolar */);
-		}
-	}
-
 	Super::PreSave(InSaveContext);
 }
 #endif // WITH_EDITOR
diff --git a/Engine/Plugins/Runtime/AudioModulation/Source/AudioModulation/Public/SoundModulationPatch.h b/Engine/Plugins/Runtime/AudioModulation/Source/AudioModulation/Public/SoundModulationPatch.h
--- a/Engine/Plugins/Runtime/AudioModulation/Source/AudioModulation/Public/SoundModulationPatch.h
+++ b/Engine/Plugins/Runtime/AudioModulation/Source/AudioModulation/Public/SoundModulationPatch.h
@@ -59,9 +59,6 @@ struct AUDIOMODULATION_API FSoundControlModulationPatch
 	/** Modulation inputs */
 	UPROPERTY(EditAnywhere, BlueprintReadWrite, Category = Inputs)
 	TArray<FSoundControlModulationInput> Inputs;
-
-	UPROPERTY(EditAnywhere, BlueprintReadWrite, Category = Inputs)
-	EWaveTableResolution WaveTableResolution = EWaveTableResolution::None;
 };
 
 UCLASS(config = Engine, editinlinenew, BlueprintType)
diff --git a/Engine/Plugins/Runtime/AudioModulation/Source/AudioModulationEditor/Private/Editors/ModulationPatchCurveEditorViewStacked.cpp b/Engine/Plugins/Runtime/AudioModulation/Source/AudioModulationEditor/Private/Editors/ModulationPatchCurveEditorViewStacked.cpp
--- a/Engine/Plugins/Runtime/AudioModulation/Source/AudioModulationEditor/Private/Editors/ModulationPatchCurveEditorViewStacked.cpp
+++ b/Engine/Plugins/Runtime/AudioModulation/Source/AudioModulationEditor/Private/Editors/ModulationPatchCurveEditorViewStacked.cpp
@@ -97,17 +97,6 @@ const USoundModulationPatch* FModPatchCurveEditorModel::GetPatch() const
 	return nullptr;
 }
 
-FWaveTableTransform* FModPatchCurveEditorModel::GetTransform()
-{
-	if (USoundModulationPatch* Patch = GetPatch())
-	{
-		FSoundControlModulationInput& Input = Patch->PatchSettings.Inputs[GetCurveIndex()];
-		return &Input.Transform;
-	}
-
-	return nullptr;
-}
-
 void FModPatchCurveEditorModel::RefreshCurveDescriptorText(const FWaveTableTransform& InTransform, FText& OutShortDisplayName, FText& OutInputAxisName, FText& OutOutputAxisName)
 {
 	OutShortDisplayName = LOCTEXT("ModulationCurveDisplayTitle_BusUnset", "Bus (Unset)");
diff --git a/Engine/Plugins/Runtime/AudioModulation/Source/AudioModulationEditor/Private/Editors/ModulationPatchCurveEditorViewStacked.h b/Engine/Plugins/Runtime/AudioModulation/Source/AudioModulationEditor/Private/Editors/ModulationPatchCurveEditorViewStacked.h
--- a/Engine/Plugins/Runtime/AudioModulation/Source/AudioModulationEditor/Private/Editors/ModulationPatchCurveEditorViewStacked.h
+++ b/Engine/Plugins/Runtime/AudioModulation/Source/AudioModulationEditor/Private/Editors/ModulationPatchCurveEditorViewStacked.h
@@ -36,8 +36,6 @@ public:
 
 	virtual void RefreshCurveDescriptorText(const FWaveTableTransform& InTransform, FText& OutShortDisplayName, FText& OutInputAxisName, FText& OutOutputAxisName) override;
 
-	virtual FWaveTableTransform* GetTransform() override;
-
 	virtual ECurveEditorViewID GetViewId() const override { return ModPatchViewId; }
 	virtual FColor GetCurveColor() const override { return UAudioModulationStyle::GetControlBusColor(); }
 	virtual bool GetPropertyEditorDisabled() const override { return GetIsBypassed(); }
diff --git a/Engine/Plugins/Runtime/AudioModulation/Source/AudioModulationEditor/Private/Editors/ModulationPatchEditor.cpp b/Engine/Plugins/Runtime/AudioModulation/Source/AudioModulationEditor/Private/Editors/ModulationPatchEditor.cpp
--- a/Engine/Plugins/Runtime/AudioModulation/Source/AudioModulationEditor/Private/Editors/ModulationPatchEditor.cpp
+++ b/Engine/Plugins/Runtime/AudioModulation/Source/AudioModulationEditor/Private/Editors/ModulationPatchEditor.cpp
@@ -20,16 +20,6 @@ TUniquePtr<WaveTable::Editor::FWaveTableCurveModel> FModulationPatchEditor::Cons
 	return TUniquePtr<FWaveTableCurveModel>(new FModPatchCurveEditorModel(InRichCurve, GetEditingObject(), InSource));
 }
 
-EWaveTableResolution FModulationPatchEditor::GetBankResolution() const
-{
-	if (USoundModulationPatch* Patch = Cast<USoundModulationPatch>(GetEditingObject()))
-	{
-		return Patch->PatchSettings.WaveTableResolution;
-	}
-
-	return EWaveTableResolution::None;
-}
-
 bool FModulationPatchEditor::GetIsPropertyEditorDisabled() const
 {
 	return GetIsBypassed();
diff --git a/Engine/Plugins/Runtime/AudioModulation/Source/AudioModulationEditor/Private/Editors/ModulationPatchEditor.h b/Engine/Plugins/Runtime/AudioModulation/Source/AudioModulationEditor/Private/Editors/ModulationPatchEditor.h
--- a/Engine/Plugins/Runtime/AudioModulation/Source/AudioModulationEditor/Private/Editors/ModulationPatchEditor.h
+++ b/Engine/Plugins/Runtime/AudioModulation/Source/AudioModulationEditor/Private/Editors/ModulationPatchEditor.h
@@ -15,8 +15,6 @@ public:
 protected:
 	virtual bool GetIsBypassed() const;
 	virtual bool GetIsPropertyEditorDisabled() const override;
-	virtual EWaveTableResolution GetBankResolution() const override;
-	virtual bool GetBankIsBipolar() const override { return false; }
 	virtual FWaveTableTransform* GetTransform(int32 InIndex) const override;
 	virtual int32 GetNumTransforms() const override;
 
diff --git a/Engine/Plugins/Runtime/AudioModulation/Source/AudioModulationEditor/Private/Layouts/SoundControlModulationPatchLayout.cpp b/Engine/Plugins/Runtime/AudioModulation/Source/AudioModulationEditor/Private/Layouts/SoundControlModulationPatchLayout.cpp
--- a/Engine/Plugins/Runtime/AudioModulation/Source/AudioModulationEditor/Private/Layouts/SoundControlModulationPatchLayout.cpp
+++ b/Engine/Plugins/Runtime/AudioModulation/Source/AudioModulationEditor/Private/Layouts/SoundControlModulationPatchLayout.cpp
@@ -72,7 +72,6 @@ void FSoundControlModulationPatchLayoutCustomization::CustomizeChildren(TSharedR
 	AudioModulationEditorUtils::GetPropertyHandleMap(StructPropertyHandle, PropertyHandles);
 
 	TSharedRef<IPropertyHandle> BypassHandle = PropertyHandles.FindChecked(GET_MEMBER_NAME_CHECKED(FSoundControlModulationPatch, bBypass)).ToSharedRef();
-	TSharedRef<IPropertyHandle> WaveTableResolutionHandle = PropertyHandles.FindChecked(GET_MEMBER_NAME_CHECKED(FSoundControlModulationPatch, WaveTableResolution)).ToSharedRef();
 	TSharedRef<IPropertyHandle> InputsHandle = PropertyHandles.FindChecked(GET_MEMBER_NAME_CHECKED(FSoundControlModulationPatch, Inputs)).ToSharedRef();
 	TSharedRef<IPropertyHandle> ParameterHandle = PropertyHandles.FindChecked(GET_MEMBER_NAME_CHECKED(FSoundControlModulationPatch, OutputParameter)).ToSharedRef();
 
@@ -84,7 +83,6 @@ void FSoundControlModulationPatchLayoutCustomization::CustomizeChildren(TSharedR
 	});
 
 	ChildBuilder.AddProperty(BypassHandle);
-	ChildBuilder.AddProperty(WaveTableResolutionHandle);
 	ChildBuilder.AddProperty(ParameterHandle).EditCondition(EditCondition, nullptr);
 	ChildBuilder.AddProperty(InputsHandle).EditCondition(EditCondition, nullptr);
 }
diff --git a/Engine/Plugins/Runtime/AudioModulation/Source/AudioModulationEditor/Private/Layouts/SoundModulationTransformLayout.cpp b/Engine/Plugins/Runtime/AudioModulation/Source/AudioModulationEditor/Private/Layouts/SoundModulationTransformLayout.cpp
--- a/Engine/Plugins/Runtime/AudioModulation/Source/AudioModulationEditor/Private/Layouts/SoundModulationTransformLayout.cpp
+++ b/Engine/Plugins/Runtime/AudioModulation/Source/AudioModulationEditor/Private/Layouts/SoundModulationTransformLayout.cpp
@@ -6,23 +6,11 @@
 #include "UObject/NoExportTypes.h"
 
 
-const EWaveTableResolution* FSoundModulationTransformLayoutCustomization::GetResolution() const
+TSet<EWaveTableCurve> FSoundModulationTransformLayoutCustomization::GetSupportedCurves() const
 {
-	if (ensure(WaveTableOptionsHandle.IsValid()))
-	{
-		TArray<UObject*> OuterObjects;
-		WaveTableOptionsHandle->GetOuterObjects(OuterObjects);
-		if (OuterObjects.Num() == 1)
-		{
-			if (USoundModulationPatch* Patch = Cast<USoundModulationPatch>(OuterObjects.Last()))
-			{
-				FSoundControlModulationPatch& PatchSettings = Patch->PatchSettings;
-				return &PatchSettings.WaveTableResolution;
-			}
-		}
-	}
-
-	return nullptr;
+	TSet<EWaveTableCurve> Curves = WaveTable::Editor::FTransformLayoutCustomizationBase::GetSupportedCurves();
+	Curves.Remove(EWaveTableCurve::File);
+	return Curves;
 }
 
 FWaveTableTransform* FSoundModulationTransformLayoutCustomization::GetTransform() const
diff --git a/Engine/Plugins/Runtime/AudioModulation/Source/AudioModulationEditor/Private/Layouts/SoundModulationTransformLayout.h b/Engine/Plugins/Runtime/AudioModulation/Source/AudioModulationEditor/Private/Layouts/SoundModulationTransformLayout.h
--- a/Engine/Plugins/Runtime/AudioModulation/Source/AudioModulationEditor/Private/Layouts/SoundModulationTransformLayout.h
+++ b/Engine/Plugins/Runtime/AudioModulation/Source/AudioModulationEditor/Private/Layouts/SoundModulationTransformLayout.h
@@ -25,7 +25,7 @@ public:
 
 protected:
 	//~ Begin FTransformLayoutCustomizationBase
-	virtual const EWaveTableResolution* GetResolution() const override;
+	virtual TSet<EWaveTableCurve> GetSupportedCurves() const override;
 	virtual FWaveTableTransform* GetTransform() const override;
 	virtual bool IsBipolar() const override;
 	//~ End FTransformLayoutCustomizationBase
diff --git a/Engine/Plugins/Runtime/WaveTable/Source/WaveTable/Private/WaveTableBank.cpp b/Engine/Plugins/Runtime/WaveTable/Source/WaveTable/Private/WaveTableBank.cpp
--- a/Engine/Plugins/Runtime/WaveTable/Source/WaveTable/Private/WaveTableBank.cpp
+++ b/Engine/Plugins/Runtime/WaveTable/Source/WaveTable/Private/WaveTableBank.cpp
@@ -1,4 +1,5 @@
 // Copyright Epic Games, Inc. All Rights Reserved.
+
 #include "WaveTableBank.h"
 
 #define LOCTEXT_NAMESPACE "WaveTable"
@@ -10,15 +11,39 @@ TUniquePtr<Audio::IProxyData> UWaveTableBank::CreateNewProxyData(const Audio::FP
 }
 
 #if WITH_EDITOR
+void UWaveTableBank::RefreshWaveTables()
+{
+	int32 MaxPCMSize = 0;
+	for (FWaveTableBankEntry& Entry : Entries)
+	{
+		TArrayView<const float> EditSourceView = Entry.Transform.WaveTableSettings.GetEditSourceView();
+		MaxPCMSize = FMath::Max(MaxPCMSize, EditSourceView.Num());
+	}
+
+	WaveTableSizeMB = 0.0f;
+	WaveTableLengthSec = 0.0f;
+	for (FWaveTableBankEntry& Entry : Entries)
+	{
+		TArray<float>& TransformWaveTable = Entry.Transform.WaveTable;
+		TransformWaveTable.Empty();
+
+		const int32 WaveTableSize = WaveTable::GetWaveTableSize(Resolution, Entry.Transform.Curve, MaxPCMSize);
+		TransformWaveTable.AddZeroed(WaveTableSize);
+
+		Entry.Transform.CreateWaveTable(TransformWaveTable, bBipolar);
+		WaveTableSizeMB += sizeof(float) * WaveTableSize;
+		WaveTableLengthSec = FMath::Max(WaveTableLengthSec, WaveTableSize / 48000.f);
+	}
+
+	WaveTableSizeMB /= FMath::Square(1024);
+}
+
 void UWaveTableBank::PostEditChangeProperty(FPropertyChangedEvent& InPropertyChangedEvent)
 {
 	FProperty* Property = InPropertyChangedEvent.Property;
 	if (Property && InPropertyChangedEvent.ChangeType != EPropertyChangeType::Interactive)
 	{
-		for (FWaveTableBankEntry& Entry : Entries)
-		{
-			Entry.Transform.CacheWaveTable(Resolution, bBipolar);
-		}
+		RefreshWaveTables();
 	}
 
 	Super::PostEditChangeProperty(InPropertyChangedEvent);
@@ -29,10 +54,7 @@ void UWaveTableBank::PostEditChangeChainProperty(FPropertyChangedChainEvent& InP
 	FProperty* Property = InPropertyChangedEvent.Property;
 	if (Property && InPropertyChangedEvent.ChangeType != EPropertyChangeType::Interactive)
 	{
-		for (FWaveTableBankEntry& Entry : Entries)
-		{
-			Entry.Transform.CacheWaveTable(Resolution, bBipolar);
-		}
+		RefreshWaveTables();
 	}
 
 	Super::PostEditChangeChainProperty(InPropertyChangedEvent);
@@ -42,10 +64,7 @@ void UWaveTableBank::PreSave(FObjectPreSaveContext InSaveContext)
 {
 	if (!InSaveContext.IsCooking())
 	{
-		for (FWaveTableBankEntry& Entry : Entries)
-		{
-			Entry.Transform.CacheWaveTable(Resolution, bBipolar);
-		}
+		RefreshWaveTables();
 	}
 
 	Super::PreSave(InSaveContext);
diff --git a/Engine/Plugins/Runtime/WaveTable/Source/WaveTable/Private/WaveTableImporter.cpp b/Engine/Plugins/Runtime/WaveTable/Source/WaveTable/Private/WaveTableImporter.cpp
--- a/Engine/Plugins/Runtime/WaveTable/Source/WaveTable/Private/WaveTableImporter.cpp
+++ b/Engine/Plugins/Runtime/WaveTable/Source/WaveTable/Private/WaveTableImporter.cpp
@@ -9,25 +9,12 @@
 #include "UObject/ObjectMacros.h"
 #include "UObject/Package.h"
 #include "WaveTableSampler.h"
-#include "DSP/EnvelopeFollower.h"
 
 
 namespace WaveTable
 {
-	namespace ImporterPrivate
-	{
-		UE_NODISCARD static FORCEINLINE int32 RatioToIndex(float InRatio, int32 InArrayNum)
-		{
-			check(InArrayNum > 0);
-			check(InRatio >= 0.0f);
-
-			return FMath::Min(InArrayNum - 1, FMath::TruncToInt32(InRatio * InArrayNum));
-		};
-	}
-
-	FImporter::FImporter(const FWaveTableSettings& InSettings, EWaveTableResolution InResolution, bool bInBipolar)
+	FImporter::FImporter(const FWaveTableSettings& InSettings, bool bInBipolar)
 		: Settings(InSettings)
-		, Resolution(InResolution)
 		, bBipolar(bInBipolar)
 	{
 		Sampler.SetInterpolationMode(FWaveTableSampler::EInterpolationMode::Cubic);
@@ -36,45 +23,13 @@ namespace WaveTable
 
 	void FImporter::Process(TArray<float>& OutWaveTable)
 	{
-		using namespace ImporterPrivate;
-
-		const int32 ResolutionInt = ResolutionToInt32(Resolution);
-
-		OutWaveTable.Reset();
-		OutWaveTable.AddZeroed(ResolutionInt);
-
-		if (Resolution == EWaveTableResolution::None)
-		{
-			return;
-		}
-
-		int32 SourceNumSamples = Settings.SourcePCMData.Num();
-		if (SourceNumSamples == 0)
+		// 1. Get editable section of source PCM
+		TArrayView<const float> EditSourceView = Settings.GetEditSourceView();
+		if (!EditSourceView.IsEmpty())
 		{
-			return;
-		}
-
-		// 1. Use Top & Tail to determine start offset & sample count of SourceArrayView
-		int32 SourceTopOffset = 0;
-		if (Settings.Top > 0.0f)
-		{
-			SourceTopOffset = ImporterPrivate::RatioToIndex(Settings.Top, SourceNumSamples);
-			SourceNumSamples -= SourceTopOffset;
-		}
-
-		if (Settings.Tail > 0.0f)
-		{
-			const int32 SourceTailOffset = RatioToIndex(Settings.Tail, SourceNumSamples);
-			SourceNumSamples -= SourceTailOffset;
-		}
-
-		if (SourceNumSamples > 0 && SourceTopOffset < Settings.SourcePCMData.Num() && !OutWaveTable.IsEmpty())
-		{
-			TArrayView<const float> SourceArrayView(Settings.SourcePCMData.GetData() + SourceTopOffset, SourceNumSamples);
-
 			// 2. Resample into table
 			Sampler.Reset();
-			Sampler.Process(SourceArrayView, OutWaveTable);
+			Sampler.Process(EditSourceView, OutWaveTable);
 
 			// 3. Apply offset
 			float TableOffset = 0.0f;
@@ -94,7 +49,7 @@ namespace WaveTable
 				}
 			}
 
-			// 6. Apply fades
+			// 5. Apply fades
 			if (Settings.FadeIn > 0.0f)
 			{
 				const int32 FadeLength = RatioToIndex(Settings.FadeIn, OutWaveTable.Num());
@@ -110,7 +65,7 @@ namespace WaveTable
 				Audio::ArrayFade(FadeView, 1.0f, 0.0f);
 			}
 
-			// 7. Finalize if unipolar source
+			// 6. Finalize if unipolar source
 			if (!bBipolar)
 			{
 				Audio::ArrayAbsInPlace(OutWaveTable);
diff --git a/Engine/Plugins/Runtime/WaveTable/Source/WaveTable/Private/WaveTableSampler.cpp b/Engine/Plugins/Runtime/WaveTable/Source/WaveTable/Private/WaveTableSampler.cpp
--- a/Engine/Plugins/Runtime/WaveTable/Source/WaveTable/Private/WaveTableSampler.cpp
+++ b/Engine/Plugins/Runtime/WaveTable/Source/WaveTable/Private/WaveTableSampler.cpp
@@ -23,6 +23,39 @@ namespace WaveTable
 			return FVector2f(2.0f, NextValue - LastValue).GetSafeNormal();
 		};
 
+		auto MaxValueIndexInterpolator = [](TArrayView<const float> InTableView, TArrayView<float> InOutIndicesToValues)
+		{
+			int32 LastIndex = 0;
+			const float* TableData = InTableView.GetData();
+			float* IndexValueData = InOutIndicesToValues.GetData();
+			for (int32 i = 0; i < InOutIndicesToValues.Num(); ++i)
+			{
+				int32 Index = FMath::TruncToInt32(InTableView.Num() * IndexValueData[i]) % InTableView.Num();
+				bool bSwapped = false;
+				if (Index < LastIndex)
+				{
+					bSwapped = true;
+					::Swap(LastIndex, Index);
+				}
+
+				int32 MaxAbs = 0.0f;
+				for (int32 j = LastIndex; j <= Index; ++j)
+				{
+					float IndexAbs = FMath::Abs(TableData[j]);
+					if (IndexAbs > MaxAbs)
+					{
+						MaxAbs = IndexAbs;
+						IndexValueData[i] = TableData[j];
+					}
+				}
+
+				if (!bSwapped)
+				{
+					LastIndex = Index;
+				}
+			}
+		};
+
 		auto CubicIndexInterpolator = [](TArrayView<const float> InTableView, TArrayView<float> InOutIndicesToValues)
 		{
 			const int32 NumTableSamples = InTableView.Num();
@@ -32,7 +65,7 @@ namespace WaveTable
 				float& IndexToOutput = InOutIndicesToValues[i];
 				const int32 LastIndexInt = FMath::TruncToInt32(IndexToOutput);
 
-				const FVector2f P0 = { (float)LastIndexInt, InTableView[LastIndexInt % NumTableSamples] };
+				const FVector2f P0 = { (float)LastIndexInt, InTable[LastIndexInt % NumTableSamples] };
 				const FVector2f P1 = { (float)(LastIndexInt + 1), InTable[(LastIndexInt + 1) % NumTableSamples] };
 				const FVector2f TangentP0 = GetTangentP0(InTable, P1.Y, IndexToOutput, NumTableSamples);
 				const FVector2f TangentP1 = GetTangentP1(InTable, P0.Y, IndexToOutput, NumTableSamples);
@@ -106,9 +139,15 @@ namespace WaveTable
 			}
 			break;
 
+			case EInterpolationMode::MaxValue:
+			{
+				MaxValueIndexInterpolator(InTableView, InOutIndexToSamplesView);
+			}
+			break;
+
 			default:
 			{
-				static_assert(static_cast<int32>(EInterpolationMode::COUNT) == 3, "Possible missing switch coverage for EInterpolationMode");
+				static_assert(static_cast<int32>(EInterpolationMode::COUNT) == 4, "Possible missing switch coverage for EInterpolationMode");
 				checkNoEntry();
 			}
 			break;
@@ -210,7 +249,6 @@ namespace WaveTable
 		return Index;
 	}
 
-
 	float FWaveTableSampler::Process(TArrayView<const float> InTableView, TArrayView<float> OutSamplesView)
 	{
 		return Process(InTableView, { }, { }, { }, OutSamplesView);
diff --git a/Engine/Plugins/Runtime/WaveTable/Source/WaveTable/Private/WaveTableSettings.cpp b/Engine/Plugins/Runtime/WaveTable/Source/WaveTable/Private/WaveTableSettings.cpp
--- a/Engine/Plugins/Runtime/WaveTable/Source/WaveTable/Private/WaveTableSettings.cpp
+++ b/Engine/Plugins/Runtime/WaveTable/Source/WaveTable/Private/WaveTableSettings.cpp
@@ -2,62 +2,99 @@
 #include "WaveTableSettings.h"
 
 
-namespace WaveTable
+TArrayView<const float> FWaveTableSettings::GetEditSourceView() const
 {
-	int32 ResolutionToInt32(EWaveTableResolution InWaveTableResolution)
+	int32 SourceNumSamples = SourcePCMData.Num();
+	if (SourceNumSamples == 0)
 	{
-		if (InWaveTableResolution == EWaveTableResolution::None)
-		{
-			return 0;
-		}
+		return { };
+	}
 
-		return 1 << static_cast<int32>(InWaveTableResolution);
+	int32 SourceTopOffset = 0;
+	if (Top > 0.0f)
+	{
+		SourceTopOffset = WaveTable::RatioToIndex(Top, SourceNumSamples);
+		SourceNumSamples -= SourceTopOffset;
 	}
 
-	int32 ResolutionToInt32(EWaveTableResolution InWaveTableResolution, EWaveTableCurve InCurve)
+	if (Tail > 0.0f)
 	{
-		if (InWaveTableResolution != EWaveTableResolution::None)
-		{
-			return ResolutionToInt32(InWaveTableResolution);
-		}
+		const int32 SourceTailOffset = WaveTable::RatioToIndex(Tail, SourceNumSamples);
+		SourceNumSamples -= SourceTailOffset;
+	}
 
-		switch (InCurve)
-		{
-			// File always uses the incoming resolution, even if 'None'
-			case EWaveTableCurve::File:
-			break;
+	if (SourceTopOffset < SourcePCMData.Num())
+	{
+		return { SourcePCMData.GetData() + SourceTopOffset, SourceNumSamples };
+	}
 
-			case EWaveTableCurve::Linear:
-			case EWaveTableCurve::Linear_Inv:
+	return { };
+}
+
+namespace WaveTable
+{
+	int32 GetWaveTableSize(EWaveTableResolution InWaveTableResolution, EWaveTableCurve InCurve, int32 InMaxPCMSize)
+	{
+		auto GetCurveResolution = [](EWaveTableCurve Curve)
+		{
+			switch (Curve)
 			{
-				InWaveTableResolution = EWaveTableResolution::Res_8;
-			}
-			break;
+				case EWaveTableCurve::File:
+				{
+					return EWaveTableResolution::Maximum;
+				}
+				break;
+
+				case EWaveTableCurve::Linear:
+				case EWaveTableCurve::Linear_Inv:
+				{
+					return EWaveTableResolution::Res_8;
+				}
+				break;
+
+				case EWaveTableCurve::Exp:
+				case EWaveTableCurve::Exp_Inverse:
+				case EWaveTableCurve::Log:
+				{
+					return EWaveTableResolution::Res_256;
+				}
+				break;
 
-			case EWaveTableCurve::Exp:
-			case EWaveTableCurve::Exp_Inverse:
-			case EWaveTableCurve::Log:
+				case EWaveTableCurve::Sin:
+				case EWaveTableCurve::SCurve:
+				case EWaveTableCurve::Sin_Full:
+				{
+					return EWaveTableResolution::Res_64;
+				}
+				break;
+
+				default:
+				{
+					return EWaveTableResolution::Res_128;
+				}
+				break;
+			};
+		};
+
+		switch (InWaveTableResolution)
+		{
+			case EWaveTableResolution::Maximum:
 			{
-				InWaveTableResolution = EWaveTableResolution::Res_256;
+				return FMath::Max(InMaxPCMSize, 1 << static_cast<int32>(EWaveTableResolution::Res_Max));
 			}
 			break;
 
-			case EWaveTableCurve::Sin:
-			case EWaveTableCurve::SCurve:
-			case EWaveTableCurve::Sin_Full:
+			case EWaveTableResolution::None:
 			{
-				InWaveTableResolution = EWaveTableResolution::Res_64;
+				const EWaveTableResolution CurveRes = GetCurveResolution(InCurve);
+				return 1 << static_cast<int32>(CurveRes);
 			}
 			break;
 
-
 			default:
 			{
-				InWaveTableResolution = EWaveTableResolution::Res_128;
+				return 1 << static_cast<int32>(InWaveTableResolution);
 			}
-			break;
-		}
-
-		return ResolutionToInt32(InWaveTableResolution);
+		};
 	}
 } // namespace WaveTable
diff --git a/Engine/Plugins/Runtime/WaveTable/Source/WaveTable/Private/WaveTableTransform.cpp b/Engine/Plugins/Runtime/WaveTable/Source/WaveTable/Private/WaveTableTransform.cpp
--- a/Engine/Plugins/Runtime/WaveTable/Source/WaveTable/Private/WaveTableTransform.cpp
+++ b/Engine/Plugins/Runtime/WaveTable/Source/WaveTable/Private/WaveTableTransform.cpp
@@ -22,17 +22,6 @@ void FWaveTableTransform::Apply(float& InOutValue, bool bInBipolar) const
 
 }
 
-void FWaveTableTransform::BuildWaveTable(TArrayView<float> InOutTable, bool bInBipolar) const
-{
-	const float Delta = 1.0f / InOutTable.Num();
-	for (int32 i = 0; i < InOutTable.Num(); ++i)
-	{
-		InOutTable[i] = i * Delta;
-	}
-
-	bInBipolar ? SampleCurveBipolar(InOutTable) : SampleCurveUnipolar(InOutTable);
-}
-
 void FWaveTableTransform::SampleCurveBipolar(TArrayView<float> InOutTable) const
 {
 	// Clamp the input
@@ -148,11 +137,8 @@ void FWaveTableTransform::SampleCurveBipolar(TArrayView<float> InOutTable) const
 			}
 			else
 			{
-				for (int32 i = 0; i < InOutTable.Num(); ++i)
-				{
-					const int32 Index = FMath::TruncToInt32(WaveTable.Num() * ValueData[i]);
-					ValueData[i] = WaveTable[Index % WaveTable.Num()];
-				}
+				constexpr auto InterpMode = WaveTable::FWaveTableSampler::EInterpolationMode::MaxValue;
+				WaveTable::FWaveTableSampler::Interpolate(WaveTable, InOutTable, InterpMode);
 			}
 		}
 		break;
@@ -277,11 +263,8 @@ void FWaveTableTransform::SampleCurveUnipolar(TArrayView<float> InOutTable) cons
 			}
 			else
 			{
-				for (int32 i = 0; i < InOutTable.Num(); ++i)
-				{
-					const int32 Index = FMath::TruncToInt32(WaveTable.Num() * ValueData[i]);
-					ValueData[i] = WaveTable[Index % WaveTable.Num()];
-				}
+				constexpr auto InterpMode = WaveTable::FWaveTableSampler::EInterpolationMode::MaxValue;
+				WaveTable::FWaveTableSampler::Interpolate(WaveTable, InOutTable, InterpMode);
 			}
 		}
 		break;
@@ -310,17 +293,25 @@ void FWaveTableTransform::CacheCurve()
 }
 
 #if WITH_EDITOR
-void FWaveTableTransform::CacheWaveTable(EWaveTableResolution InResolution, bool bInBipolar)
+void FWaveTableTransform::CopyToWaveTable(TArrayView<float> InOutTable, bool bInBipolar) const
 {
-	WaveTable.Empty();
+	const float Delta = 1.0f / InOutTable.Num();
+	for (int32 i = 0; i < InOutTable.Num(); ++i)
+	{
+		InOutTable[i] = i * Delta;
+	}
+
+	bInBipolar ? SampleCurveBipolar(InOutTable) : SampleCurveUnipolar(InOutTable);
+}
 
-	// If not wavetable, just sample using apply function.
+void FWaveTableTransform::CreateWaveTable(TArray<float>& InOutTable, bool bInBipolar) const
+{
 	switch (Curve)
 	{
 		case EWaveTableCurve::File:
 		{
-			WaveTable::FImporter Importer(WaveTableSettings, InResolution, bInBipolar);
-			Importer.Process(WaveTable);
+			WaveTable::FImporter Importer(WaveTableSettings, bInBipolar);
+			Importer.Process(InOutTable);
 		}
 		break;
 
@@ -336,8 +327,7 @@ void FWaveTableTransform::CacheWaveTable(EWaveTableResolution InResolution, bool
 		case EWaveTableCurve::Sin_Full:
 		default:
 		{
-			WaveTable.AddZeroed(WaveTable::ResolutionToInt32(InResolution));
-			BuildWaveTable(WaveTable, bInBipolar);
+			CopyToWaveTable(InOutTable, bInBipolar);
 		}
 		break;
 	}
diff --git a/Engine/Plugins/Runtime/WaveTable/Source/WaveTable/Public/WaveTableBank.h b/Engine/Plugins/Runtime/WaveTable/Source/WaveTable/Public/WaveTableBank.h
--- a/Engine/Plugins/Runtime/WaveTable/Source/WaveTable/Public/WaveTableBank.h
+++ b/Engine/Plugins/Runtime/WaveTable/Source/WaveTable/Public/WaveTableBank.h
@@ -27,10 +27,9 @@ class WAVETABLE_API UWaveTableBank : public UObject, public IAudioProxyDataFacto
 	GENERATED_BODY()
 
 public:
-	// If set to 'None', WaveTable is not cached.  If set to
-	// value, caches a sampled wavetable for all curves.
+	// Number of sampled cached for each curve in the given bank.
 	UPROPERTY(EditAnywhere, Category = Options)
-	EWaveTableResolution Resolution;
+	EWaveTableResolution Resolution = EWaveTableResolution::Res_256;
 
 	// Determines if output from curve/wavetable are to be clamped between [-1.0f, 1.0f]
 	// (i.e. for waveform generation, oscillation, etc.) or between [0.0f, 1.0f]
@@ -38,6 +37,16 @@ public:
 	UPROPERTY(EditAnywhere, Category = Options)
 	bool bBipolar = true;
 
+#if WITH_EDITORONLY_DATA
+	// Sum total size of all WaveTable data within the given bank
+	UPROPERTY(VisibleAnywhere, Category = Options, meta = (DisplayName = "WaveTable Size (MB)"))
+	float WaveTableSizeMB = 0.0f;
+
+	// Length of all WaveTable samples in bank in seconds (at 48kHz)
+	UPROPERTY(VisibleAnywhere, Category = Options, meta = (DisplayName = "WaveTable Length (sec)"))
+	float WaveTableLengthSec = 0.0f;
+#endif // WITH_EDITORONLY_DATA
+
 	/** Tables within the given bank */
 	UPROPERTY(EditAnywhere, Category = Options)
 	TArray<FWaveTableBankEntry> Entries;
@@ -46,6 +55,8 @@ public:
 	virtual TUniquePtr<Audio::IProxyData> CreateNewProxyData(const Audio::FProxyDataInitParams& InitParams) override;
 
 #if WITH_EDITOR
+	void RefreshWaveTables();
+
 	virtual void PreSave(FObjectPreSaveContext InSaveContext) override;
 	virtual void PostEditChangeProperty(FPropertyChangedEvent& InPropertyChangedEvent) override;
 	virtual void PostEditChangeChainProperty(FPropertyChangedChainEvent& InPropertyChangedEvent) override;
diff --git a/Engine/Plugins/Runtime/WaveTable/Source/WaveTable/Public/WaveTableImporter.h b/Engine/Plugins/Runtime/WaveTable/Source/WaveTable/Public/WaveTableImporter.h
--- a/Engine/Plugins/Runtime/WaveTable/Source/WaveTable/Public/WaveTableImporter.h
+++ b/Engine/Plugins/Runtime/WaveTable/Source/WaveTable/Public/WaveTableImporter.h
@@ -11,14 +11,13 @@ namespace WaveTable
 	class WAVETABLE_API FImporter
 	{
 	public:
-		FImporter(const FWaveTableSettings& InSettings, EWaveTableResolution InResolution, bool bInBipolar);
+		FImporter(const FWaveTableSettings& InSettings, bool bInBipolar);
 
 		void Process(TArray<float>& OutWaveTable);
 
 	private:
 		const FWaveTableSettings& Settings;
 
-		EWaveTableResolution Resolution;
 		bool bBipolar = false;
 
 		FWaveTableSampler Sampler;
diff --git a/Engine/Plugins/Runtime/WaveTable/Source/WaveTable/Public/WaveTableSampler.h b/Engine/Plugins/Runtime/WaveTable/Source/WaveTable/Public/WaveTableSampler.h
--- a/Engine/Plugins/Runtime/WaveTable/Source/WaveTable/Public/WaveTableSampler.h
+++ b/Engine/Plugins/Runtime/WaveTable/Source/WaveTable/Public/WaveTableSampler.h
@@ -12,17 +12,26 @@ namespace WaveTable
 	class WAVETABLE_API FWaveTableSampler
 	{
 	public:
-		enum class EInterpolationMode : uint8
+		// Interpolation mode when sampling values between indices
+		enum class EInterpolationMode
 		{
+			// No interpolation (Stepped)
 			None,
+
+			// Linear interpolation
 			Linear,
+
+			// Cubic interpolation
 			Cubic,
 
+			// Takes the maximum value between two points (EXPENSIVE, but good for caching/drawn curves).
+			MaxValue,
+
 			COUNT
 		};
 
-		// Mode of interpolation between last value in table and subsequent input index.
-		enum class ESingleSampleMode : uint8
+		// Mode of interpolation between last value in table and subsequent input index when sampling single values.
+		enum class ESingleSampleMode
 		{
 			// Interpolates last value to zero (0.0f) in table if sampled index is beyond last position
 			Zero = 0,
diff --git a/Engine/Plugins/Runtime/WaveTable/Source/WaveTable/Public/WaveTableSettings.h b/Engine/Plugins/Runtime/WaveTable/Source/WaveTable/Public/WaveTableSettings.h
--- a/Engine/Plugins/Runtime/WaveTable/Source/WaveTable/Public/WaveTableSettings.h
+++ b/Engine/Plugins/Runtime/WaveTable/Source/WaveTable/Public/WaveTableSettings.h
@@ -2,6 +2,8 @@
 #pragma once
 
 #include "Containers/Array.h"
+#include "Containers/ArrayView.h"
+#include "GenericPlatform/GenericPlatformMath.h"
 #include "HAL/Platform.h"
 #include "UObject/SoftObjectPath.h"
 #include "UObject/Object.h"
@@ -10,10 +12,24 @@
 #include "WaveTableSettings.generated.h"
 
 
+namespace WaveTable
+{
+	// Converts a ratio [0.0, 1.0] to an index within the given number of array elements.
+	UE_NODISCARD static FORCEINLINE int32 RatioToIndex(float InRatio, int32 InArrayNum)
+	{
+		check(InArrayNum > 0);
+		check(InRatio >= 0.0f);
+
+		return FMath::Min(InArrayNum - 1, FMath::TruncToInt32(InRatio * InArrayNum));
+	};
+} // namespace WaveTable
+
+
 UENUM(BlueprintType)
 enum class EWaveTableResolution : uint8
 {
-	None		= 0 UMETA(DisplayName = "None"),
+	None		= 0 UMETA(Hidden),
+
 	Res_8		= 3 UMETA(DisplayName = "8"),
 	Res_16		= 4 UMETA(DisplayName = "16"),
 	Res_32		= 5 UMETA(DisplayName = "32"),
@@ -25,7 +41,14 @@ enum class EWaveTableResolution : uint8
 	Res_2048	= 11 UMETA(DisplayName = "2048"),
 	Res_4096	= 12 UMETA(DisplayName = "4096"),
 
-	MAX = Res_4096 UMETA(Hidden)
+	// Maximum "static" power-of-two WaveTable resolution
+	Res_Max		= Res_4096 UMETA(Hidden),
+
+	// Takes largest source file length (if WaveTable was produced
+	// by a file, PCM stream, SoundWave, etc.) of superset of waves,
+	// or Res_Max if all superset of waves were produced by curves. This
+	// can be very memory intensive if source file is large!
+	Maximum
 };
 
 
@@ -43,23 +66,23 @@ struct WAVETABLE_API FWaveTableSettings
 	int32 ChannelIndex = 0;
 
 	// Percent to phase shift of table
-	UPROPERTY(EditAnywhere, Category = Options, meta = (ClampMin = "0.0", ClampMax = "1.0", SliderExponent = 0.5))
+	UPROPERTY(EditAnywhere, Category = Options, meta = (ClampMin = "0.0", ClampMax = "1.0"))
 	float Phase = 0.0f;
 
 	// Percent to remove from beginning of sampled WaveTable.
-	UPROPERTY(EditAnywhere, Category = Options, meta = (ClampMin = "0.0", ClampMax = "1.0", SliderExponent = 0.5))
+	UPROPERTY(EditAnywhere, Category = Options, meta = (ClampMin = "0.0", ClampMax = "1.0"))
 	float Top = 0.0f;
 
 	// Percent to remove from end of sampled WaveTable.
-	UPROPERTY(EditAnywhere, Category = Options, meta = (ClampMin = "0.0", ClampMax = "1.0", SliderExponent = 0.5))
+	UPROPERTY(EditAnywhere, Category = Options, meta = (ClampMin = "0.0", ClampMax = "1.0"))
 	float Tail = 0.0f;
 
 	// Percent to fade in over.
-	UPROPERTY(EditAnywhere, Category = Options, meta = (ClampMin = "0.0", ClampMax = "1.0", SliderExponent = 0.5))
+	UPROPERTY(EditAnywhere, Category = Options, meta = (ClampMin = "0.0", ClampMax = "1.0"))
 	float FadeIn  = 0.0f;
 
 	// Percent to fade out over.
-	UPROPERTY(EditAnywhere, Category = Options, meta = (ClampMin = "0.0", ClampMax = "1.0", SliderExponent = 0.5))
+	UPROPERTY(EditAnywhere, Category = Options, meta = (ClampMin = "0.0", ClampMax = "1.0"))
 	float FadeOut = 0.0f;
 
 	// Whether or not to normalize the WaveTable.
@@ -74,6 +97,11 @@ struct WAVETABLE_API FWaveTableSettings
 	// SourcePCM Data
 	UPROPERTY()
 	TArray<float> SourcePCMData;
+
+	// Returns a TArrayView of Source PCM values that are on when
+	// importing with the given settings (restricted effectively
+	// due to top/tail settings).
+	TArrayView<const float> GetEditSourceView() const;
 };
 
 UENUM(BlueprintType)
@@ -103,10 +131,7 @@ enum class EWaveTableCurve : uint8
 
 namespace WaveTable
 {
-	// Converts WaveTableResolution to integer resolution value
-	WAVETABLE_API int32 ResolutionToInt32(EWaveTableResolution InWaveTableResolution);
-
-	// Converts WaveTableResolution to integer resolution value.  If WaveTableResolution is set to none, uses
-	// default value associated with provided curve.
-	WAVETABLE_API int32 ResolutionToInt32(EWaveTableResolution InWaveTableResolution, EWaveTableCurve InCurve);
+	// Returns WaveTable size given WaveTableResolution and optional Curve or PCM size. If resolution is none, uses default value associated
+	// with provided curve. If set to file, will use either the largest static resolution (Res_Max, or 4096) or InMaxPCMSize, whichever is larger.
+	WAVETABLE_API int32 GetWaveTableSize(EWaveTableResolution InWaveTableResolution, EWaveTableCurve InCurve, int32 InMaxPCMSize = INDEX_NONE);
 } // namespace WaveTable
\ No newline at end of file
diff --git a/Engine/Plugins/Runtime/WaveTable/Source/WaveTable/Public/WaveTableTransform.h b/Engine/Plugins/Runtime/WaveTable/Source/WaveTable/Public/WaveTableTransform.h
--- a/Engine/Plugins/Runtime/WaveTable/Source/WaveTable/Public/WaveTableTransform.h
+++ b/Engine/Plugins/Runtime/WaveTable/Source/WaveTable/Public/WaveTableTransform.h
@@ -48,15 +48,19 @@ struct WAVETABLE_API FWaveTableTransform
 	FWaveTableSettings WaveTableSettings;
 #endif // WITH_EDITORONLY_DATA
 
+	// Applies transform to provided index value, setting it to the corresponding value.
 	void Apply(float& InOutValue, bool bInBipolar = false) const;
 
-	// Builds a single period of a WaveTable using this transform in place of the provided ArrayView.
-	// Size of the WaveTable is dictated by the provided ArrayView's size.
-	void BuildWaveTable(TArrayView<float> InOutTable, bool bInBipolar) const;
-
 #if WITH_EDITOR
-	// No-ops if curve is set to WaveTable.
-	void CacheWaveTable(EWaveTableResolution InResolution, bool bInBipolar);
+	// Copies a single, period [0.0, 1.0] to a WaveTable using this transform's curve. The
+	// size of the resulting WaveTable is the provided ArrayView's size. If curve is itself
+	// set to 'File' and internally cached as WaveTable, will re-sample using max value
+	// interpolation. Does *not* apply destructive edits for using WaveTableSettings PCM data.
+	void CopyToWaveTable(TArrayView<float> InOutTable, bool bInBipolar) const;
+
+	// Caches WaveTable for curve using the provided array length as resolution.
+	// Applies destructive edits as provided by WaveTableSettings if curve set to 'File'.
+	void CreateWaveTable(TArray<float>& InOutTable, bool bInBipolar) const;
 #endif // WITH_EDITOR
 
 	/** Caches curve data.  Should a shared curve be selected, curve points are copied locally
diff --git a/Engine/Plugins/Runtime/WaveTable/Source/WaveTableEditor/Private/WaveTableBankEditor.cpp b/Engine/Plugins/Runtime/WaveTable/Source/WaveTableEditor/Private/WaveTableBankEditor.cpp
--- a/Engine/Plugins/Runtime/WaveTable/Source/WaveTableEditor/Private/WaveTableBankEditor.cpp
+++ b/Engine/Plugins/Runtime/WaveTable/Source/WaveTableEditor/Private/WaveTableBankEditor.cpp
@@ -39,6 +39,22 @@ namespace WaveTable
 {
 	namespace Editor
 	{
+		namespace BankEditorPrivate
+		{
+			// Two seconds of samples at 48kHz
+			constexpr int32 MaxLiveEditSamples = 48000 * 2;
+
+			int32 GetDrawPointCount(EWaveTableResolution Resolution, const FWaveTableTransform& InTransform)
+			{
+				// Min/max number of curve points for mathematical
+				// expressions/files drawn on curve stack view.
+				static const FVector2D ExpressionCurvePointCountLimits = { 2, 4096 };
+
+				float WaveTableSize = WaveTable::GetWaveTableSize(Resolution, InTransform.Curve);
+				return FMath::Clamp(WaveTableSize, ExpressionCurvePointCountLimits.X, ExpressionCurvePointCountLimits.Y);
+			}
+		} // namespace BankEditorPrivate
+
 		const FName FBankEditorBase::AppIdentifier(TEXT("WaveTableEditorApp"));
 		const FName FBankEditorBase::CurveTabId(TEXT("WaveTableEditor_Curves"));
 		const FName FBankEditorBase::PropertiesTabId(TEXT("WaveTableEditor_Properties"));
@@ -99,36 +115,36 @@ namespace WaveTable
 			PropertiesView->SetObject(InParentObject);
 
 			TSharedRef<FTabManager::FLayout> StandaloneDefaultLayout = FTabManager::NewLayout("Standalone_WaveTableEditor_Layout_v1")
-				->AddArea
+			->AddArea
+			(
+				FTabManager::NewPrimaryArea()
+				->SetOrientation(Orient_Vertical)
+				->Split
 				(
-					FTabManager::NewPrimaryArea()
-					->SetOrientation(Orient_Vertical)
+					FTabManager::NewSplitter()
+					->SetSizeCoefficient(0.9f)
+					->SetOrientation(Orient_Horizontal)
+					->Split
+					(
+						FTabManager::NewStack()
+						->SetSizeCoefficient(0.225f)
+						->AddTab(PropertiesTabId, ETabState::OpenedTab)
+					)
 					->Split
 					(
 						FTabManager::NewSplitter()
-						->SetSizeCoefficient(0.9f)
-						->SetOrientation(Orient_Horizontal)
+						->SetSizeCoefficient(0.775f)
+						->SetOrientation(Orient_Vertical)
 						->Split
 						(
 							FTabManager::NewStack()
-							->SetSizeCoefficient(0.225f)
-							->AddTab(PropertiesTabId, ETabState::OpenedTab)
-						)
-						->Split
-						(
-							FTabManager::NewSplitter()
-							->SetSizeCoefficient(0.775f)
-							->SetOrientation(Orient_Vertical)
-							->Split
-							(
-								FTabManager::NewStack()
-								->SetHideTabWell(true)
-								->SetSizeCoefficient(0.33f)
-								->AddTab(CurveTabId, ETabState::OpenedTab)
-							)
+							->SetHideTabWell(true)
+							->SetSizeCoefficient(0.33f)
+							->AddTab(CurveTabId, ETabState::OpenedTab)
 						)
 					)
-				);
+				)
+			);
 
 			const bool bCreateDefaultStandaloneMenu = true;
 			const bool bCreateDefaultToolbar = true;
@@ -145,9 +161,6 @@ namespace WaveTable
 				bToolbarFocusable,
 				bUseSmallIcons);
 
-		// 	Editor::FModule& WaveTableModule = FModuleManager::LoadModuleChecked<Editor::FModule>(TEXT("WaveTableEditor"));
-		// 	AddMenuExtender(WaveTableModule.GetWaveTableMenuExtensibilityManager()->GetAllExtenders(GetToolkitCommands(), GetEditingObjects()));
-		// 	AddToolbarExtender(WaveTableModule.GetWaveTableToolbarExtensibilityManager()->GetAllExtenders(GetToolkitCommands(), GetEditingObjects()));
 			AddToolbarExtender(CurvePanel->GetToolbarExtender());
 
 			if (CurveEditor.IsValid())
@@ -164,7 +177,7 @@ namespace WaveTable
 			}
 		}
 
-		void FBankEditorBase::GenerateExpressionCurve(WaveTable::Editor::FBankEditorBase::FCurveData& OutCurveData, int32 InTransformIndex, EWaveTableCurveSource InSource, bool bInIsUnset)
+		void FBankEditorBase::GenerateExpressionCurve(WaveTable::Editor::FBankEditorBase::FCurveData& OutCurveData, int32 InTransformIndex, bool bInIsUnset)
 		{
 			TSharedPtr<FRichCurve>& Curve = OutCurveData.ExpressionCurve;
 			if (!Curve.IsValid())
@@ -176,55 +189,58 @@ namespace WaveTable
 
 			if (!GetIsPropertyEditorDisabled())
 			{
-				Curve->Reset();
-
 				if (const FWaveTableTransform* Transform = GetTransform(InTransformIndex))
 				{
-					// Need at least first and last point to generate curve
+					// Need at least first and last point to generate expression curve
 					const EWaveTableResolution BankResolution = GetBankResolution();
-					const int32 PointCount = FMath::Max(WaveTable::ResolutionToInt32(BankResolution, Transform->Curve), 2);
 
-					TArray<FRichCurveKey> NewKeys;
+					const int32 PointCount = BankEditorPrivate::GetDrawPointCount(BankResolution, *Transform);
+					TArray<float> KeyTable;
+					KeyTable.AddZeroed(PointCount);
 
-					if (BankResolution == EWaveTableResolution::None)
+					// Optimization for really big source files if set to 'File'. Do not recreate WaveTable on the fly
+					// for large samples as caching mechanism is too slow for rapid regeneration (i.e. when making
+					// interactive edits).
+					if (Transform->WaveTableSettings.SourcePCMData.Num() < BankEditorPrivate::MaxLiveEditSamples)
 					{
-						// Terminate curves at last constant point for non-wavetables
-						TArray<float> KeyTable;
-						KeyTable.AddZeroed(PointCount);
-						Transform->BuildWaveTable(KeyTable, bIsBipolar);
+						Transform->CreateWaveTable(KeyTable, bIsBipolar);
+					}
+					else
+					{
+						Transform->CopyToWaveTable(KeyTable, bIsBipolar);
+					}
 
-						const float Delta = 1.0 / KeyTable.Num();
-						for (int32 i = 0; i < KeyTable.Num(); ++i)
-						{
-							NewKeys.Add({ Delta * i, KeyTable[i] });
-						}
+					const float Delta = 1.0 / KeyTable.Num();
+
+					FRichCurve& CurveRef = *Curve.Get();
+					TArray<FRichCurveKey> NewKeys;
+
+					// Terminate curves at initial point for wavetables
+					for (int32 i = 0; i < KeyTable.Num(); ++i)
+					{
+						NewKeys.Add({ Delta * i, KeyTable[i] });
+					}
 
+					// Terminate curves at initial point for wavetables
+					// (as they loop) and at the final computed location
+					// for non-wavetables (as they are envelopes).
+					if (BankResolution == EWaveTableResolution::None)
+					{
 						float TermPoint = 1.0f;
 						Transform->Apply(TermPoint, bIsBipolar);
 						NewKeys.Add({ 1.0f, TermPoint });
 					}
 					else
 					{
-						// Terminate curves at initial point for wavetables
-						TArray<float> KeyTable;
-						KeyTable.AddZeroed(PointCount);
-						Transform->BuildWaveTable(KeyTable, bIsBipolar);
-
-						const float Delta = 1.0 / KeyTable.Num();
-						for (int32 i = 0; i < KeyTable.Num(); ++i)
-						{
-							NewKeys.Add({ Delta * i, KeyTable[i] });
-						}
-
 						NewKeys.Add({ 1.0f, KeyTable[0] });
 					}
 
-					Curve->SetKeys(NewKeys);
+					CurveRef.SetKeys(NewKeys);
+
+					const EWaveTableCurveSource Source = bInIsUnset ? EWaveTableCurveSource::Unset : EWaveTableCurveSource::Expression;
+					SetCurve(InTransformIndex, *Curve.Get(), Source);
 				}
 			}
-
-			const EWaveTableCurveSource Source = bInIsUnset ? EWaveTableCurveSource::Unset : EWaveTableCurveSource::Expression;
-			SetCurve(InTransformIndex, *Curve.Get(), Source);
 		}
 
 		bool FBankEditorBase::RequiresNewCurve(int32 InTransformIndex, const FRichCurve& InRichCurve) const
@@ -263,21 +279,21 @@ namespace WaveTable
 
 			FCurveData& CurveDataEntry = CurveData[InTransformIndex];
 
+			const bool bIsBipolar = GetBankIsBipolar();
 			const bool bRequiresNewCurve = RequiresNewCurve(InTransformIndex, InRichCurve);
 			if (bRequiresNewCurve)
 			{
 				TUniquePtr<FWaveTableCurveModel> NewCurve = ConstructCurveModel(InRichCurve, GetEditingObject(), InSource);
-				NewCurve->Refresh(*Transform, InTransformIndex);
+				NewCurve->Refresh(*Transform, InTransformIndex, bIsBipolar);
 				CurveDataEntry.ModelID = CurveEditor->AddCurve(MoveTemp(NewCurve));
+				CurveEditor->PinCurve(CurveDataEntry.ModelID);
 			}
 			else
 			{
 				const TUniquePtr<FCurveModel>& CurveModel = CurveEditor->GetCurves().FindChecked(CurveDataEntry.ModelID);
 				check(CurveModel.Get());
-				static_cast<FWaveTableCurveModel*>(CurveModel.Get())->Refresh(*Transform, InTransformIndex);
+				static_cast<FWaveTableCurveModel*>(CurveModel.Get())->Refresh(*Transform, InTransformIndex, bIsBipolar);
 			}
-
-			CurveEditor->PinCurve(CurveDataEntry.ModelID);
 		}
 
 		FName FBankEditorBase::GetToolkitFName() const
@@ -313,6 +329,12 @@ namespace WaveTable
 			{
 				RefreshCurves();
 			}
+			else
+			{
+				// TODO: Isolate file that changed and only update that curve for performance
+				// (Could get bad for larger banks with many file curves).
+				RegenerateFileCurves();
+			}
 		}
 
 		TSharedRef<SDockTab> FBankEditorBase::SpawnTab_Properties(const FSpawnTabArgs& Args)
@@ -426,6 +448,25 @@ namespace WaveTable
 			}
 		}
 
+		void FBankEditorBase::RegenerateFileCurves()
+		{
+			for (int32 i = 0; i < GetNumTransforms(); ++i)
+			{
+				if (CurveData.IsValidIndex(i))
+				{
+					FCurveData& CurveDataEntry = CurveData[i];
+
+					if (FWaveTableTransform* Transform = GetTransform(i))
+					{
+						if (Transform->Curve == EWaveTableCurve::File)
+						{
+							GenerateExpressionCurve(CurveDataEntry, i);
+						}
+					}
+				}
+			}
+		}
+
 		void FBankEditorBase::RefreshCurves()
 		{
 			check(CurveEditor.IsValid());
@@ -477,7 +518,7 @@ namespace WaveTable
 					case EWaveTableCurve::Sin_Full:
 					case EWaveTableCurve::File:
 					{
-						GenerateExpressionCurve(CurveDataEntry, i, EWaveTableCurveSource::Expression);
+						GenerateExpressionCurve(CurveDataEntry, i);
 					}
 					break;
 
@@ -492,7 +533,7 @@ namespace WaveTable
 						{
 							// Builds a dummy expression that just maps input to output in case
 							// where asset isn't selected and leave source as unset
-							GenerateExpressionCurve(CurveDataEntry, i, EWaveTableCurveSource::Expression, true /* bIsUnset */);
+							GenerateExpressionCurve(CurveDataEntry, i, true /* bIsUnset */);
 						}
 					}
 					break;
@@ -522,18 +563,27 @@ namespace WaveTable
 				ActiveModelIDs.Add(CurveDataEntry.ModelID);
 			}
 
-			const TMap<FCurveModelID, TUniquePtr<FCurveModel>>& Curves = CurveEditor->GetCurves();
-			for (const TPair<FCurveModelID, TUniquePtr<FCurveModel>>& Pair : Curves)
+			// Remove all dead curves
 			{
-				if (!ActiveModelIDs.Contains(Pair.Key))
+				const TMap<FCurveModelID, TUniquePtr<FCurveModel>>& Curves = CurveEditor->GetCurves();
+				for (const TPair<FCurveModelID, TUniquePtr<FCurveModel>>& Pair : Curves)
 				{
-					ToRemove.Add(Pair.Key);
+					if (!ActiveModelIDs.Contains(Pair.Key))
+					{
+						ToRemove.Add(Pair.Key);
+					}
+				}
+
+				for (FCurveModelID ModelID : ToRemove)
+				{
+					CurveEditor->RemoveCurve(ModelID);
 				}
 			}
 
-			for (FCurveModelID ModelID : ToRemove)
+			// Re-pin all curves
+			for (const FCurveData& CurveDataEntry : CurveData)
 			{
-				CurveEditor->RemoveCurve(ModelID);
+				CurveEditor->PinCurve(CurveDataEntry.ModelID);
 			}
 		}
 
diff --git a/Engine/Plugins/Runtime/WaveTable/Source/WaveTableEditor/Private/WaveTableCurveEditorViewStacked.cpp b/Engine/Plugins/Runtime/WaveTable/Source/WaveTableEditor/Private/WaveTableCurveEditorViewStacked.cpp
--- a/Engine/Plugins/Runtime/WaveTable/Source/WaveTableEditor/Private/WaveTableCurveEditorViewStacked.cpp
+++ b/Engine/Plugins/Runtime/WaveTable/Source/WaveTableEditor/Private/WaveTableCurveEditorViewStacked.cpp
@@ -10,6 +10,7 @@
 #include "SCurveEditorPanel.h"
 #include "SGraphActionMenu.h"
 #include "Styling/AppStyle.h"
+#include "Views/SInteractiveCurveEditorView.h"
 #include "WaveTableTransform.h"
 #include "Widgets/Text/STextBlock.h"
 
@@ -28,7 +29,7 @@ namespace WaveTable
 			check(InOwner);
 		}
 
-		void FWaveTableCurveModel::Refresh(const FWaveTableTransform& InTransform, int32 InCurveIndex)
+		void FWaveTableCurveModel::Refresh(const FWaveTableTransform& InTransform, int32 InCurveIndex, bool bInIsBipolar)
 		{
 			// Must be set prior to remainder of refresh to avoid
 			// child class implementation acting on incorrect cached index
@@ -42,12 +43,26 @@ namespace WaveTable
 
 			AxesDescriptor = FText::Format(LOCTEXT("WaveTableCurveDisplayTitle_AxesFormat", "X = {0}, Y = {1}"), InputAxisName, OutputAxisName);
 
-			bKeyDrawEnabled = true;
 			IntentionName = TEXT("AudioControlValue");
 			SupportedViews = GetViewId();
 
 			bKeyDrawEnabled = false;
 
+			bIsBipolar = bInIsBipolar;
+
+			if (InTransform.Curve == EWaveTableCurve::File)
+			{
+				NumSamples = InTransform.WaveTable.Num();
+				FadeInRatio = InTransform.WaveTableSettings.FadeIn;
+				FadeOutRatio = InTransform.WaveTableSettings.FadeOut;
+			}
+			else
+			{
+				NumSamples = 0;
+				FadeInRatio = 0.0f;
+				FadeOutRatio = 0.0f;
+			}
+
 			const bool bIsDisabled = GetPropertyEditorDisabled();
 			if (bIsDisabled)
 			{
@@ -55,9 +70,6 @@ namespace WaveTable
 			}
 			else
 			{
-				UEnum* Enum = StaticEnum<EWaveTableCurveSource>();
-				check(Enum);
-
 				LongDisplayName = FText::Format(LOCTEXT("WaveTableCurveDisplayTitle_NameFormat", "{0}: {1}"), FText::AsNumber(CurveIndex), ShortDisplayName);
 
 				if (Source == EWaveTableCurveSource::Custom)
@@ -68,16 +80,12 @@ namespace WaveTable
 				{
 					bKeyDrawEnabled = true;
 
-					if (FWaveTableTransform* Transform = GetTransform())
+					if (const UCurveFloat* SharedCurve = InTransform.CurveShared)
 					{
-						if (const UCurveFloat* SharedCurve = Transform->CurveShared)
-						{
-							const FText CurveSourceText = FText::FromString(SharedCurve->GetName());
-							LongDisplayName = FText::Format(LOCTEXT("WaveTableCurveDisplayTitle_NameSharedFormat", "{0}: {1}, Shared Curve '{2}'"), FText::AsNumber(CurveIndex), ShortDisplayName, CurveSourceText);
-						}
+						const FText CurveSourceText = FText::FromString(SharedCurve->GetName());
+						LongDisplayName = FText::Format(LOCTEXT("WaveTableCurveDisplayTitle_NameSharedFormat", "{0}: {1}, Shared Curve '{2}'"), FText::AsNumber(CurveIndex), ShortDisplayName, CurveSourceText);
 					}
 				}
-
 			}
 		}
 
@@ -88,6 +96,20 @@ namespace WaveTable
 				: Color.Desaturate(0.35f);
 		}
 
+		void FWaveTableCurveModel::GetValueRange(double& MinValue, double& MaxValue) const
+		{
+			if (FadeInRatio > 0.0f || FadeOutRatio > 0.0f)
+			{
+				MaxValue = 1.0f;
+
+				MinValue = bIsBipolar ? -1.0f : 0.0f;
+			}
+			else
+			{
+				FRichCurveEditorModelRaw::GetValueRange(MinValue, MaxValue);
+			}
+		}
+
 		const FText& FWaveTableCurveModel::GetAxesDescriptor() const
 		{
 			return AxesDescriptor;
@@ -103,11 +125,6 @@ namespace WaveTable
 			return ParentObject.Get();
 		}
 
-		FWaveTableTransform* FWaveTableCurveModel::GetTransform()
-		{
-			return nullptr;
-		}
-
 		bool FWaveTableCurveModel::IsReadOnly() const
 		{
 			return Source != EWaveTableCurveSource::Custom;
@@ -170,8 +187,6 @@ namespace WaveTable
 				return;
 			}
 
-			const int32 LabelLayerId = BaseLayerId + CurveViewConstants::ELayerOffset::Labels;
-
 			const double ValuePerPixel = 1.0 / StackedHeight;
 			const double ValueSpacePadding = StackedPadding * ValuePerPixel;
 
@@ -210,6 +225,7 @@ namespace WaveTable
 				FPaintGeometry LabelGeometry = AllottedGeometry.ToPaintGeometry(FSlateLayoutTransform(LabelPosition));
 				const FVector2D LabelSize = FontMeasure->Measure(Label, FontInfo);
 
+				const uint32 LabelLayerId = BaseLayerId + CurveViewConstants::ELayerOffset::GridLabels;
 				FSlateDrawElement::MakeText(OutDrawElements, LabelLayerId + 1, LabelGeometry, Label, FontInfo, DrawEffects, Curve->GetColor());
 
 				// Render axes descriptor
@@ -251,10 +267,16 @@ namespace WaveTable
 			const double ValuePerPixel = 1.0 / StackedHeight;
 			const double ValueSpacePadding = StackedPadding * ValuePerPixel;
 
+			if (CurveInfoByID.IsEmpty())
+			{
+				return;
+			}
+
 			for (auto It = CurveInfoByID.CreateConstIterator(); It; ++It)
 			{
 				DrawInfo.SetCurveModel(CurveEditor->FindCurve(It.Key()));
-				if (!ensureAlways(DrawInfo.GetCurveModel()))
+				const FCurveModel* CurveModel = DrawInfo.GetCurveModel();
+				if (!ensureAlways(CurveModel))
 				{
 					continue;
 				}
@@ -262,7 +284,7 @@ namespace WaveTable
 				TArray<FText> CurveModelGridLabelsX = MajorGridLabelsX;
 				check(MajorGridLinesX.Num() == CurveModelGridLabelsX.Num());
 
-				if (const FWaveTableCurveModel* EditorModel = static_cast<const FWaveTableCurveModel*>(DrawInfo.GetCurveModel()))
+				if (const FWaveTableCurveModel* EditorModel = static_cast<const FWaveTableCurveModel*>(CurveModel))
 				{
 					for (int32 i = 0; i < CurveModelGridLabelsX.Num(); ++i)
 					{
@@ -324,17 +346,67 @@ namespace WaveTable
 						DrawViewGridLineY(VerticalLine, OutDrawElements, DrawInfo, DrawEffects, nullptr /* Label */, false /* bIsMajor */);
 					}
 				}
+
+				if (const FWaveTableCurveModel* EditorModel = static_cast<const FWaveTableCurveModel*>(CurveModel))
+				{
+					auto DrawFade = [this, &CurveEditor, &DrawInfo, &OutDrawElements, DrawEffects, BaseLayerId, ModelID = It.Key()](float StartPosX, float EndPosX, bool bFadeIn, bool bIsBipolar)
+					{
+						const SCurveEditorView* View = CurveEditor->FindFirstInteractiveView(ModelID);
+						FCurveEditorScreenSpace CurveSpace = View->GetCurveSpace(ModelID);
+
+						StartPosX = CurveSpace.SecondsToScreen(StartPosX);
+						EndPosX = CurveSpace.SecondsToScreen(EndPosX);
+						static const FLinearColor FadeColor = FLinearColor::White.CopyWithNewOpacity(0.75f);
+
+						DrawInfo.LinePoints[0].X = StartPosX;
+						DrawInfo.LinePoints[0].Y = CurveSpace.ValueToScreen(0.0f);
+						DrawInfo.LinePoints[1].X = EndPosX;
+						DrawInfo.LinePoints[1].Y = CurveSpace.ValueToScreen(1.0f);
+						const int32 FadeLayerId = BaseLayerId + CurveViewConstants::ELayerOffset::Curves;
+
+						FSlateDrawElement::MakeLines(OutDrawElements, FadeLayerId, DrawInfo.PaintGeometry, DrawInfo.LinePoints, DrawEffects, FadeColor, false);
+
+						if (bIsBipolar)
+						{
+							DrawInfo.LinePoints[1].Y = CurveSpace.ValueToScreen(-1.0f);
+							FSlateDrawElement::MakeLines(OutDrawElements, FadeLayerId, DrawInfo.PaintGeometry, DrawInfo.LinePoints, DrawEffects, FadeColor, false);
+						}
+
+						const FSlateBrush* WhiteBrush = FAppStyle::GetBrush("WhiteBrush");
+
+						const float Width = FMath::Abs(DrawInfo.LinePoints[1].X - DrawInfo.LinePoints[0].X);
+						const float BoxStart = bFadeIn ? StartPosX : StartPosX - Width;
+						const FPaintGeometry BoxGeometry = DrawInfo.AllottedGeometry->ToPaintGeometry(
+							FVector2D(Width, StackedHeight),
+							FSlateLayoutTransform(FVector2D(BoxStart, DrawInfo.GetPixelTop()))
+						);
+
+						const int32 GridOverlayLayerId = DrawInfo.GetBaseLayerId() + CurveViewConstants::ELayerOffset::GridOverlays;
+						FSlateDrawElement::MakeBox(OutDrawElements, GridOverlayLayerId, BoxGeometry, WhiteBrush, DrawEffects, FadeColor.CopyWithNewOpacity(0.15f));
+					};
+
+					const bool bIsBipolar = EditorModel->GetIsBipolar();
+					float FadeInRatio = EditorModel->GetFadeInRatio();
+					if (FadeInRatio > 0.0f)
+					{
+						DrawFade(0.0f, FadeInRatio, true, bIsBipolar);
+					}
+
+					float FadeOutRatio = EditorModel->GetFadeOutRatio();
+					if (FadeOutRatio > 0.0f)
+					{
+						DrawFade(1.0f, 1.0f - FadeOutRatio, false, bIsBipolar);
+					}
+				}
 			}
 		}
 
 		void SViewStacked::DrawViewGridLineX(FSlateWindowElementList& OutDrawElements, FGridDrawInfo& DrawInfo, ESlateDrawEffect DrawEffects, double OffsetAlpha, bool bIsMajor) const
 		{
-			double ValueMin;
-			double ValueMax;
+			double ValueMin = 0.0f;
+			double ValueMax = 0.0f;
 			DrawInfo.GetCurveModel()->GetValueRange(ValueMin, ValueMax);
 
-			const int32 GridLineLayerId = DrawInfo.GetBaseLayerId() + CurveViewConstants::ELayerOffset::GridLines;
-
 			const double LowerValue = DrawInfo.GetLowerValue();
 			const double PixelBottom = DrawInfo.GetPixelBottom();
 			const double PixelTop = DrawInfo.GetPixelTop();
@@ -342,10 +414,12 @@ namespace WaveTable
 			const TSharedRef<FSlateFontMeasure> FontMeasure = FSlateApplication::Get().GetRenderer()->GetFontMeasureService();
 			const FSlateFontInfo FontInfo = FCoreStyle::Get().GetFontStyle("CurveEd.InfoFont");
 
+			const uint32 LineLayerId = DrawInfo.GetBaseLayerId() + CurveViewConstants::ELayerOffset::GridLines;
+
 			FLinearColor Color = bIsMajor ? DrawInfo.GetMajorGridColor() : DrawInfo.GetMinorGridColor();
 
 			DrawInfo.LinePoints[0].Y = DrawInfo.LinePoints[1].Y = DrawInfo.ScreenSpace.ValueToScreen(LowerValue + OffsetAlpha);
-			FSlateDrawElement::MakeLines(OutDrawElements, GridLineLayerId, DrawInfo.PaintGeometry, DrawInfo.LinePoints, DrawEffects, Color, false);
+			FSlateDrawElement::MakeLines(OutDrawElements, LineLayerId, DrawInfo.PaintGeometry, DrawInfo.LinePoints, DrawEffects, Color, false);
 
 			FText Label = FText::AsNumber(FMath::Lerp(ValueMin, ValueMax, OffsetAlpha), &DrawInfo.LabelFormat);
 
@@ -369,10 +443,11 @@ namespace WaveTable
 			}
 
 			const FPaintGeometry LabelGeometry = DrawInfo.AllottedGeometry->ToPaintGeometry(FSlateLayoutTransform(FVector2D(CurveViewConstants::CurveLabelOffsetX, LabelY)));
+			const uint32 LabelLayerId = DrawInfo.GetBaseLayerId() + CurveViewConstants::ELayerOffset::GridLabels;
 
 			FSlateDrawElement::MakeText(
 				OutDrawElements,
-				DrawInfo.GetBaseLayerId() + CurveViewConstants::ELayerOffset::GridLabels,
+				LabelLayerId,
 				LabelGeometry,
 				Label,
 				FontInfo,
@@ -386,12 +461,11 @@ namespace WaveTable
 			const float Width = DrawInfo.AllottedGeometry->GetLocalSize().X;
 			if (VerticalLine >= 0 || VerticalLine <= FMath::RoundToFloat(Width))
 			{
-				const int32 GridLineLayerId = DrawInfo.GetBaseLayerId() + CurveViewConstants::ELayerOffset::GridLines;
-
 				const FLinearColor LineColor = bIsMajor ? DrawInfo.GetMajorGridColor() : DrawInfo.GetMinorGridColor();
 
 				DrawInfo.LinePoints[0].X = DrawInfo.LinePoints[1].X = VerticalLine;
-				FSlateDrawElement::MakeLines(OutDrawElements, GridLineLayerId, DrawInfo.PaintGeometry, DrawInfo.LinePoints, DrawEffects, LineColor, false);
+				const uint32 LineLayerId = DrawInfo.GetBaseLayerId() + CurveViewConstants::ELayerOffset::GridLines;
+				FSlateDrawElement::MakeLines(OutDrawElements, LineLayerId, DrawInfo.PaintGeometry, DrawInfo.LinePoints, DrawEffects, LineColor, false);
 
 				if (Label)
 				{
@@ -400,10 +474,11 @@ namespace WaveTable
 					const TSharedRef<FSlateFontMeasure> FontMeasure = FSlateApplication::Get().GetRenderer()->GetFontMeasureService();
 					const FVector2D LabelSize = FontMeasure->Measure(*Label, FontInfo);
 					const FPaintGeometry LabelGeometry = DrawInfo.AllottedGeometry->ToPaintGeometry(FSlateLayoutTransform(FVector2D(VerticalLine, DrawInfo.LinePoints[0].Y - LabelSize.Y * 1.2f)));
+					const uint32 LabelLayerId = DrawInfo.GetBaseLayerId() + CurveViewConstants::ELayerOffset::GridLabels;
 
 					FSlateDrawElement::MakeText(
 						OutDrawElements,
-						DrawInfo.GetBaseLayerId() + CurveViewConstants::ELayerOffset::GridLabels,
+						LabelLayerId,
 						LabelGeometry,
 						*Label,
 						FontInfo,
@@ -414,6 +489,20 @@ namespace WaveTable
 			}
 		}
 
+		void SViewStacked::FormatInputLabel(const FWaveTableCurveModel& EditorModel, const FNumberFormattingOptions& InLabelFormat, FText& InOutLabel) const
+		{
+			const int32 NumSamples = EditorModel.GetNumSamples();
+			if (NumSamples != 0)
+			{
+				const float Value = FCString::Atof(*InOutLabel.ToString());
+				if (Value >= 0 && Value <= 1.0f)
+				{
+					const int32 ValueTrunc = FMath::FloorToInt32(NumSamples * Value);
+					InOutLabel = FText::Format(LOCTEXT("WaveTableCurveEditor_InputLabelFormat", "{0} ({1})"), InOutLabel, FText::AsNumber(ValueTrunc));
+				}
+			}
+		}
+
 		SViewStacked::FGridDrawInfo::FGridDrawInfo(const FGeometry* InAllottedGeometry, const FCurveEditorScreenSpace& InScreenSpace, FLinearColor InGridColor, int32 InBaseLayerId)
 			: AllottedGeometry(InAllottedGeometry)
 			, ScreenSpace(InScreenSpace)
diff --git a/Engine/Plugins/Runtime/WaveTable/Source/WaveTableEditor/Private/WaveTableTransformLayout.cpp b/Engine/Plugins/Runtime/WaveTable/Source/WaveTableEditor/Private/WaveTableTransformLayout.cpp
--- a/Engine/Plugins/Runtime/WaveTable/Source/WaveTableEditor/Private/WaveTableTransformLayout.cpp
+++ b/Engine/Plugins/Runtime/WaveTable/Source/WaveTableEditor/Private/WaveTableTransformLayout.cpp
@@ -11,7 +11,9 @@
 #include "SCurveEditor.h"
 #include "ScopedTransaction.h"
 #include "Styling/AppStyle.h"
+#include "Subsystems/AssetEditorSubsystem.h"
 #include "WaveTableBank.h"
+#include "WaveTableBankEditor.h"
 #include "WaveTableFileUtilities.h"
 #include "Widgets/Input/SButton.h"
 #include "Widgets/Input/SEditableTextBox.h"
@@ -19,11 +21,12 @@
 #include "Widgets/Layout/SSeparator.h"
 #include "Widgets/Layout/SWidgetSwitcher.h"
 #include "Widgets/Notifications/SNotificationList.h"
+#include "Widgets/SToolTip.h"
 #include "Widgets/Text/SRichTextBlock.h"
 #include "Widgets/Text/STextBlock.h"
 #include "Widgets/Views/SListView.h"
 
-#define LOCTEXT_NAMESPACE "WaveTable"
+#define LOCTEXT_NAMESPACE "WaveTableEditor"
 
 
 namespace WaveTable
@@ -52,7 +55,7 @@ namespace WaveTable
 			}
 
 			CurveHandle = PropertyHandles.FindChecked(GET_MEMBER_NAME_CHECKED(FWaveTableTransform, Curve)).ToSharedRef();
-			ChildBuilder.AddProperty(CurveHandle->AsShared());
+			CustomizeCurveSelector(ChildBuilder);
 
 			TSharedRef<IPropertyHandle> ScalarHandle = PropertyHandles.FindChecked(GET_MEMBER_NAME_CHECKED(FWaveTableTransform, Scalar)).ToSharedRef();
 
@@ -75,18 +78,6 @@ namespace WaveTable
 			uint32 NumWaveTableOptions = 0;
 			if (WaveTableOptionsHandle->GetNumChildren(NumWaveTableOptions))
 			{
-				for (int32 i = 0; i < static_cast<int32>(NumWaveTableOptions); ++i)
-				{
-					TSharedPtr<IPropertyHandle> ChildProperty = WaveTableOptionsHandle->GetChildHandle(i);
-					if (ChildProperty->GetProperty()->GetFName() != GET_MEMBER_NAME_CHECKED(FWaveTableSettings, FilePath))
-					{
-						ChildProperty->SetOnPropertyValueChanged(FSimpleDelegate::CreateLambda([this]
-						{
-							RefreshWaveTable();
-						}));
-					}
-				}
-
 				ChannelIndexHandle = WaveTableOptionsHandle->GetChildHandle(GET_MEMBER_NAME_CHECKED(FWaveTableSettings, ChannelIndex));
 
 				FilePathHandle = WaveTableOptionsHandle->GetChildHandle(GET_MEMBER_NAME_CHECKED(FWaveTableSettings, FilePath));
@@ -99,6 +90,83 @@ namespace WaveTable
 				.Visibility(IsWaveTableVisibilityAttribute);
 		}
 
+		void FTransformLayoutCustomizationBase::CustomizeCurveSelector(IDetailChildrenBuilder& ChildBuilder)
+		{
+			check(CurveHandle.IsValid());
+
+			auto GetAll = [this, SupportedCurves = GetSupportedCurves()](TArray<TSharedPtr<FString>>& OutNames, TArray<TSharedPtr<SToolTip>>& OutTooltips, TArray<bool>&)
+			{
+				CurveDisplayStringToNameMap.Reset();
+				UEnum* Enum = StaticEnum<EWaveTableCurve>();
+				check(Enum);
+
+				static const int64 MaxValue = static_cast<int64>(EWaveTableCurve::Count);
+				for (int32 i = 0; i < Enum->NumEnums(); ++i)
+				{
+					EWaveTableCurve Curve = static_cast<EWaveTableCurve>(Enum->GetValueByIndex(i));
+					if (SupportedCurves.Contains(Curve))
+					{
+						TSharedPtr<FString> DisplayString = MakeShared<FString>(Enum->GetDisplayNameTextByIndex(i).ToString());
+
+						const FName Name = Enum->GetNameByIndex(i);
+						CurveDisplayStringToNameMap.Add(*DisplayString.Get(), Name);
+
+						OutTooltips.Emplace(SNew(SToolTip).Text(Enum->GetToolTipTextByIndex(i)));
+						OutNames.Emplace(MoveTemp(DisplayString));
+					}
+				}
+			};
+
+			auto GetValue = [this]()
+			{
+				UEnum* Enum = StaticEnum<EWaveTableCurve>();
+				check(Enum);
+
+				uint8 IntValue;
+				if (CurveHandle->GetValue(IntValue) == FPropertyAccess::Success)
+				{
+					return Enum->GetDisplayNameTextByValue(IntValue).ToString();
+				}
+
+				return Enum->GetDisplayNameTextByValue(static_cast<int32>(EWaveTableCurve::Count)).ToString();
+			};
+
+			auto SelectedValue = [this](const FString& InSelected)
+			{
+				UEnum* Enum = StaticEnum<EWaveTableCurve>();
+				check(Enum);
+
+				const FName& Name = CurveDisplayStringToNameMap.FindChecked(InSelected);
+				const uint8 Value = static_cast<uint8>(Enum->GetValueByName(Name, EGetByNameFlags::ErrorIfNotFound));
+				ensure(CurveHandle->SetValue(Value) == FPropertyAccess::Success);
+			};
+
+			static const FText CurveDisplayName = LOCTEXT("WaveTableTransformCurveProperty", "Curve");
+			ChildBuilder.AddCustomRow(CurveDisplayName)
+				.NameContent()
+				[
+					CurveHandle->CreatePropertyNameWidget()
+				]
+				.ValueContent()
+				.MinDesiredWidth(150.0f)
+				[
+					SNew(SHorizontalBox)
+					+ SHorizontalBox::Slot()
+					.VAlign(VAlign_Center)
+				[
+					SNew(SBox)
+					[
+						PropertyCustomizationHelpers::MakePropertyComboBox(
+							CurveHandle,
+							FOnGetPropertyComboBoxStrings::CreateLambda(GetAll),
+							FOnGetPropertyComboBoxValue::CreateLambda(GetValue),
+							FOnPropertyComboBoxValueSelected::CreateLambda(SelectedValue)
+						)
+					]
+				]
+			];
+		}
+
 		void FTransformLayoutCustomizationBase::CachePCMFromFile()
 		{
 			if (FWaveTableTransform* Transform = GetTransform())
@@ -122,17 +190,19 @@ namespace WaveTable
 				}
 			}
 
-			RefreshWaveTable();
-		}
-
-		void FTransformLayoutCustomizationBase::RefreshWaveTable() const
-		{
-			if (FWaveTableTransform* Transform = GetTransform())
+			if (GEditor)
 			{
-				if (const EWaveTableResolution* Resolution = GetResolution())
+				if (UAssetEditorSubsystem* AssetEditorSubsystem = GEditor->GetEditorSubsystem<UAssetEditorSubsystem>())
 				{
-					const bool bBipolar = IsBipolar();
-					Transform->CacheWaveTable(*Resolution, bBipolar);
+					TArray<UObject*> OuterObjects;
+					WaveTableOptionsHandle->GetOuterObjects(OuterObjects);
+					for (UObject* Object : OuterObjects)
+					{
+						if (UWaveTableBank* Bank = Cast<UWaveTableBank>(Object))
+						{
+							Bank->RefreshWaveTables();
+						}
+					}
 				}
 			}
 		}
@@ -172,22 +242,25 @@ namespace WaveTable
 			return TableIndex;
 		}
 
-		const EWaveTableResolution* FTransformLayoutCustomization::GetResolution() const
+		TSet<EWaveTableCurve> FTransformLayoutCustomizationBase::GetSupportedCurves() const
 		{
-			if (ensure(WaveTableOptionsHandle.IsValid()))
+			UEnum* Enum = StaticEnum<EWaveTableCurve>();
+			check(Enum);
+
+			TSet<EWaveTableCurve> Curves;
+			static const int64 MaxValue = static_cast<int64>(EWaveTableCurve::Count);
+			for (int32 i = 0; i < Enum->NumEnums(); ++i)
 			{
-				TArray<UObject*> OuterObjects;
-				WaveTableOptionsHandle->GetOuterObjects(OuterObjects);
-				if (OuterObjects.Num() == 1)
+				if (Enum->GetValueByIndex(i) < MaxValue)
 				{
-					if (UWaveTableBank* Bank = Cast<UWaveTableBank>(OuterObjects.Last()))
+					if (!Enum->HasMetaData(TEXT("Hidden"), i))
 					{
-						return &Bank->Resolution;
+						Curves.Add(static_cast<EWaveTableCurve>(Enum->GetValueByIndex(i)));
 					}
 				}
 			}
 
-			return nullptr;
+			return Curves;
 		}
 
 		bool FTransformLayoutCustomization::IsBipolar() const
@@ -234,4 +307,4 @@ namespace WaveTable
 		}
 	} // namespace Editor
 } // namespace WaveTable
-#undef LOCTEXT_NAMESPACE // AudioModulation
+#undef LOCTEXT_NAMESPACE // WaveTableEditor
diff --git a/Engine/Plugins/Runtime/WaveTable/Source/WaveTableEditor/Public/WaveTableBankEditor.h b/Engine/Plugins/Runtime/WaveTable/Source/WaveTableEditor/Public/WaveTableBankEditor.h
--- a/Engine/Plugins/Runtime/WaveTable/Source/WaveTableEditor/Public/WaveTableBankEditor.h
+++ b/Engine/Plugins/Runtime/WaveTable/Source/WaveTableEditor/Public/WaveTableBankEditor.h
@@ -45,10 +45,19 @@ namespace WaveTable
 			/** FNotifyHook interface */
 			virtual void NotifyPostChange(const FPropertyChangedEvent& PropertyChangedEvent, FProperty* PropertyThatChanged) override;
 
+			// Regenerates curves set to "file" without refreshing whole stack view
+			void RegenerateFileCurves();
+
+			/** Updates & redraws curves. */
+			void RefreshCurves();
+
 		protected:
 			struct FCurveData
 			{
 				FCurveModelID ModelID;
+
+				/* Curve used purely for display.  May be down-sampled from
+				 * asset's curve for performance while editing */
 				TSharedPtr<FRichCurve> ExpressionCurve;
 
 				FCurveData()
@@ -62,28 +71,34 @@ namespace WaveTable
 
 			virtual bool GetIsPropertyEditorDisabled() const { return false; }
 
-			virtual EWaveTableResolution GetBankResolution() const = 0;
-			virtual bool GetBankIsBipolar() const = 0;
+			// Returns resolution of bank being editor. By default, bank does not support WaveTable generation
+			// by being set to no resolution.
+			virtual EWaveTableResolution GetBankResolution() const { return EWaveTableResolution::None; }
 
+			// Returns whether or not bank is bipolar.  By default, returns false (bank is unipolar), functionally
+			// operating as a unipolar envelope editor.
+			virtual bool GetBankIsBipolar() const { return false; }
+
+			// Construct a new curve model for the given FRichCurve.  Allows for editor implementation to construct custom curve model types.
 			virtual TUniquePtr<FWaveTableCurveModel> ConstructCurveModel(FRichCurve& InRichCurve, UObject* InParentObject, EWaveTableCurveSource InSource) = 0;
 
+			// Returns the transform associated with the given index
 			virtual FWaveTableTransform* GetTransform (int32 InIndex) const = 0;
 
+			// Returns the number of transforms associated with the given bank.
 			virtual int32 GetNumTransforms() const = 0;
 
 			void SetCurve(int32 InTransformIndex, FRichCurve& InRichCurve, EWaveTableCurveSource InSource);
 
 		private:
-			/** Generates expression curve at the given index. */
-			void GenerateExpressionCurve(FCurveData& OutCurveData, int32 InTransformIndex, EWaveTableCurveSource InSource, bool bInIsUnset = false);
+			// Regenerates expression curve at the given index. If curve is a 'File' and source PCM data in Transform
+			// is too long, inline edits are not included in recalculation for better interact performance.
+			void GenerateExpressionCurve(FCurveData& OutCurveData, int32 InTransformIndex, bool bInIsUnset = false);
 
 			void InitCurves();
 
 			void ResetCurves();
 
-			/** Updates & redraws curves. */
-			void RefreshCurves();
-
 			/**	Spawns the tab allowing for editing/viewing the output curve(s) */
 			TSharedRef<SDockTab> SpawnTab_OutputCurve(const FSpawnTabArgs& Args);
 
@@ -117,7 +132,7 @@ namespace WaveTable
 			static const FName PropertiesTabId;
 		};
 
-		class FBankEditor : public WaveTable::Editor::FBankEditorBase
+		class WAVETABLEEDITOR_API FBankEditor : public FBankEditorBase
 		{
 		public:
 			FBankEditor() = default;
@@ -129,7 +144,7 @@ namespace WaveTable
 			virtual int32 GetNumTransforms() const override;
 			virtual FWaveTableTransform* GetTransform(int32 InIndex) const override;
 
-			virtual TUniquePtr<WaveTable::Editor::FWaveTableCurveModel> ConstructCurveModel(FRichCurve& InRichCurve, UObject* InParentObject, EWaveTableCurveSource InSource) override;
+			virtual TUniquePtr<FWaveTableCurveModel> ConstructCurveModel(FRichCurve& InRichCurve, UObject* InParentObject, EWaveTableCurveSource InSource) override;
 		};
 	} // namespace Editor
 } // namespace WaveTable
\ No newline at end of file
diff --git a/Engine/Plugins/Runtime/WaveTable/Source/WaveTableEditor/Public/WaveTableCurveEditorViewStacked.h b/Engine/Plugins/Runtime/WaveTable/Source/WaveTableEditor/Public/WaveTableCurveEditorViewStacked.h
--- a/Engine/Plugins/Runtime/WaveTable/Source/WaveTableEditor/Public/WaveTableCurveEditorViewStacked.h
+++ b/Engine/Plugins/Runtime/WaveTable/Source/WaveTableEditor/Public/WaveTableCurveEditorViewStacked.h
@@ -42,17 +42,21 @@ namespace WaveTable
 			const FText& GetAxesDescriptor() const;
 			const UObject* GetParentObject() const;
 			EWaveTableCurveSource GetSource() const;
-			void Refresh(const FWaveTableTransform& InTransform, int32 InCurveIndex);
+			void Refresh(const FWaveTableTransform& InTransform, int32 InCurveIndex, bool bInIsBipolar);
 
 			virtual ECurveEditorViewID GetViewId() const { return WaveTableViewId; }
 			virtual bool IsReadOnly() const override;
 			virtual FLinearColor GetColor() const override;
+			virtual void GetValueRange(double& MinValue, double& MaxValue) const override;
 
 			int32 GetCurveIndex() const { return CurveIndex; }
+			bool GetIsBipolar() const { return bIsBipolar; }
+			float GetFadeInRatio() const { return FadeInRatio; }
+			float GetFadeOutRatio() const { return FadeOutRatio; }
+			int32 GetNumSamples() const { return NumSamples; }
 
 		protected:
 			virtual void RefreshCurveDescriptorText(const FWaveTableTransform& InTransform, FText& OutShortDisplayName, FText& OutInputAxisName, FText& OutOutputAxisName);
-			virtual FWaveTableTransform* GetTransform();
 			virtual FColor GetCurveColor() const;
 			virtual bool GetPropertyEditorDisabled() const;
 			virtual FText GetPropertyEditorDisabledText() const;
@@ -61,8 +65,15 @@ namespace WaveTable
 
 		private:
 			int32 CurveIndex = INDEX_NONE;
+			int32 NumSamples = 0;
+
 			EWaveTableCurveSource Source = EWaveTableCurveSource::Unset;
 
+			bool bIsBipolar = false;
+
+			float FadeInRatio = 0.0f;
+			float FadeOutRatio = 0.0f;
+
 			FText InputAxisName;
 			FText AxesDescriptor;
 		};
@@ -85,7 +96,7 @@ namespace WaveTable
 			virtual void DrawViewGrids(const FGeometry& AllottedGeometry, const FSlateRect& MyCullingRect, FSlateWindowElementList& OutDrawElements, int32 BaseLayerId, ESlateDrawEffect DrawEffects) const override;
 			virtual void DrawLabels(const FGeometry& AllottedGeometry, const FSlateRect& MyCullingRect, FSlateWindowElementList& OutDrawElements, int32 BaseLayerId, ESlateDrawEffect DrawEffects) const override;
 
-			virtual void FormatInputLabel(const FWaveTableCurveModel& EditorModel, const FNumberFormattingOptions& InLabelFormat, FText& InOutLabel) const { }
+			virtual void FormatInputLabel(const FWaveTableCurveModel& EditorModel, const FNumberFormattingOptions& InLabelFormat, FText& InOutLabel) const;
 			virtual void FormatOutputLabel(const FWaveTableCurveModel& EditorModel, const FNumberFormattingOptions& InLabelFormat, FText& InOutLabel) const { }
 
 		private:
diff --git a/Engine/Plugins/Runtime/WaveTable/Source/WaveTableEditor/Public/WaveTableTransformLayout.h b/Engine/Plugins/Runtime/WaveTable/Source/WaveTableEditor/Public/WaveTableTransformLayout.h
--- a/Engine/Plugins/Runtime/WaveTable/Source/WaveTableEditor/Public/WaveTableTransformLayout.h
+++ b/Engine/Plugins/Runtime/WaveTable/Source/WaveTableEditor/Public/WaveTableTransformLayout.h
@@ -1,10 +1,12 @@
 // Copyright Epic Games, Inc. All Rights Reserved.
 #pragma once
 
+#include "Containers/Map.h"
 #include "Curves/SimpleCurve.h"
 #include "IPropertyTypeCustomization.h"
 #include "Layout/Visibility.h"
 #include "Misc/Attribute.h"
+#include "PropertyCustomizationHelpers.h"
 #include "PropertyHandle.h"
 #include "SCurveEditor.h"
 #include "WaveTableSettings.h"
@@ -24,22 +26,24 @@ namespace WaveTable
 			//~ End IPropertyTypeCustomization
 
 		protected:
-			EWaveTableCurve GetCurve() const;
-
-			virtual const EWaveTableResolution* GetResolution() const = 0;
+			virtual TSet<EWaveTableCurve> GetSupportedCurves() const;
 			virtual FWaveTableTransform* GetTransform() const = 0;
 			virtual bool IsBipolar() const = 0;
 
 			void CachePCMFromFile();
+			EWaveTableCurve GetCurve() const;
 			int32 GetOwningArrayIndex() const;
-
 			bool IsScaleableCurve() const;
-			void RefreshWaveTable() const;
 
 			TSharedPtr<IPropertyHandle> CurveHandle;
 			TSharedPtr<IPropertyHandle> ChannelIndexHandle;
 			TSharedPtr<IPropertyHandle> FilePathHandle;
 			TSharedPtr<IPropertyHandle> WaveTableOptionsHandle;
+
+		private:
+			void CustomizeCurveSelector(IDetailChildrenBuilder& ChildBuilder);
+
+			TMap<FString, FName> CurveDisplayStringToNameMap;
 		};
 
 		class WAVETABLEEDITOR_API FTransformLayoutCustomization : public FTransformLayoutCustomizationBase
@@ -50,7 +54,6 @@ namespace WaveTable
 				return MakeShared<FTransformLayoutCustomization>();
 			}
 
-			virtual const EWaveTableResolution* GetResolution() const override;
 			virtual bool IsBipolar() const override;
 			virtual FWaveTableTransform* GetTransform() const override;
 		};
