Fix missing filepath filter
#rnx

diff --git a/Engine/Plugins/Runtime/WaveTable/Source/WaveTable/Public/WaveTableSettings.h b/Engine/Plugins/Runtime/WaveTable/Source/WaveTable/Public/WaveTableSettings.h
--- a/Engine/Plugins/Runtime/WaveTable/Source/WaveTable/Public/WaveTableSettings.h
+++ b/Engine/Plugins/Runtime/WaveTable/Source/WaveTable/Public/WaveTableSettings.h
@@ -58,7 +58,7 @@ struct WAVETABLE_API FWaveTableSettings
 	GENERATED_USTRUCT_BODY()
 
 	// File to import
-	UPROPERTY(EditAnywhere, Category = Options)
+	UPROPERTY(EditAnywhere, Category = Options, meta = (FilePathFilter = "wav"))
 	FFilePath FilePath;
 
 	// Index of channel in file to build WaveTable from (wraps if channel is greater than number in file)
