[MetaSounds] Fix sync not working in wavetable sampler.
- Minor clean-up & description updates
- Make sync advanced pin
#rnx

diff --git a/Engine/Plugins/Runtime/Metasound/Source/MetasoundEngine/Private/MetasoundWaveTableOscillatorNode.cpp b/Engine/Plugins/Runtime/Metasound/Source/MetasoundEngine/Private/MetasoundWaveTableOscillatorNode.cpp
--- a/Engine/Plugins/Runtime/Metasound/Source/MetasoundEngine/Private/MetasoundWaveTableOscillatorNode.cpp
+++ b/Engine/Plugins/Runtime/Metasound/Source/MetasoundEngine/Private/MetasoundWaveTableOscillatorNode.cpp
@@ -28,11 +28,16 @@ namespace Metasound
 
 			static const FVertexInterface DefaultInterface(
 				FInputVertexInterface(
-					TInputDataVertex<FWaveTable>("WaveTable", FDataVertexMetadata { LOCTEXT("MetasoundWaveTableOscillatorNode_InputWaveTable", "WaveTable") }),
-					TInputDataVertex<FTrigger>("Play", FDataVertexMetadata { LOCTEXT("MetasoundWaveTableOscillatorNode_InputPlay", "Play") }),
-					TInputDataVertex<FTrigger>("Stop", FDataVertexMetadata { LOCTEXT("MetasoundWaveTableOscillatorNode_InputStop", "Stop") }),
-					TInputDataVertex<FTrigger>("Sync", FDataVertexMetadata { LOCTEXT("MetasoundWaveTableOscillatorNode_InputSync", "Sync") }),
-					TInputDataVertex<float>("Freq", FDataVertexMetadata { LOCTEXT("MetasoundWaveTableOscillatorNode_Freq", "Freq") }, 440.0f),
+					TInputDataVertex<FTrigger>("Play", FDataVertexMetadata { LOCTEXT("MetasoundWaveTableOscillatorNode_InputPlayDesc", "Plays the oscillator (block rate)") }),
+					TInputDataVertex<FTrigger>("Stop", FDataVertexMetadata { LOCTEXT("MetasoundWaveTableOscillatorNode_InputStopDesc", "Stops the oscillator (block rate)") }),
+					TInputDataVertex<FWaveTable>("WaveTable", FDataVertexMetadata { LOCTEXT("MetasoundWaveTableOscillatorNode_InputWaveTableDesc", "WaveTable") }),
+					TInputDataVertex<FTrigger>("Sync", FDataVertexMetadata
+					{
+						LOCTEXT("MetasoundWaveTableOscillatorNode_InputSyncDesc", "Restarts playing the WaveTable on the trigger boundary (sample rate)"),
+						LOCTEXT("MetasoundWaveTableOscillatorNode_InputSyncName", "Sync"),
+						true /* bIsAdvancedDisplay */
+					}),
+					TInputDataVertex<float>("Freq", FDataVertexMetadata { LOCTEXT("MetasoundWaveTableOscillatorNode_FreqDesc", "Frequency (number of times to sample one period of wavetable per second) [-20000Hz, 20000Hz]") }, 440.0f),
 					TInputDataVertex<FAudioBuffer>("PhaseMod", FDataVertexMetadata
 					{
 						LOCTEXT("MetasoundWaveTableOscillatorNode_PhaseModDescription", "Modulation audio source for modulating oscillation phase of provided table. A value of 0 is no phase modulation and 1 a full table length (360 degrees) of phase shift."),
diff --git a/Engine/Plugins/Runtime/WaveTable/Source/WaveTable/Private/WaveTableSampler.cpp b/Engine/Plugins/Runtime/WaveTable/Source/WaveTable/Private/WaveTableSampler.cpp
--- a/Engine/Plugins/Runtime/WaveTable/Source/WaveTable/Private/WaveTableSampler.cpp
+++ b/Engine/Plugins/Runtime/WaveTable/Source/WaveTable/Private/WaveTableSampler.cpp
@@ -295,8 +295,6 @@ namespace WaveTable
 
 	void FWaveTableSampler::ComputeIndexFrequency(TArrayView<const float> InTableView, TArrayView<const float> InFreqModulator, TArrayView<const float> InSyncTriggers, TArrayView<float> OutIndicesView)
 	{
-		float NextIndex = 0.0f;
-
 		Audio::ArraySetToConstantInplace(OutIndicesView, Settings.Freq * InTableView.Num() / OutIndicesView.Num());
 
 		if (InFreqModulator.Num() == OutIndicesView.Num())
@@ -304,50 +302,53 @@ namespace WaveTable
 			Audio::ArrayMultiplyInPlace(InFreqModulator, OutIndicesView);
 		}
 
+		auto TransformLastIndex = [this, InTableView, &InFreqModulator, &InSyncTriggers, &OutIndicesView](float InputIndex)
+		{
+			if (InFreqModulator.Num() == OutIndicesView.Num())
+			{
+				return (InFreqModulator.Last() * Settings.Freq * InTableView.Num()) + InputIndex;
+			}
+			else
+			{
+				checkf(InFreqModulator.IsEmpty(), TEXT("FreqModulator view should be the same size as the sample view or not supplied (size of 0)."));
+				return (Settings.Freq * InTableView.Num()) + InputIndex;
+			}
+		};
+
+		float* OutSamples = OutIndicesView.GetData();
 		if (InSyncTriggers.Num() == OutIndicesView.Num())
 		{
-			float* OutSamples = OutIndicesView.GetData();
-			float SyncIndex = 0;
+			float SyncIndex = 0.0f;
 			const float* SyncTrigData = InSyncTriggers.GetData();
 			for (int32 i = 0; i < OutIndicesView.Num(); ++i)
 			{
-				SyncIndex *= -1.0f * SyncTrigData[i] + 1.0f;
+				// Numerically negate trig value as high (1.0f) should
+				// be zero to "reset" index state (avoids conditional
+				// in loop at expense of memory).
+				const float InvertedTrig = -1 * SyncTrigData[i] + 1;
+				SyncIndex *= InvertedTrig;
+				LastIndex *= InvertedTrig;
 				OutSamples[i] *= SyncIndex;
 				OutSamples[i] += LastIndex;
+
 				SyncIndex++;
 			}
 
-			// Only adjust last index (indicating an inter-block reset)
-			// if the final index is less than the buffer size
-			if (SyncIndex < InSyncTriggers.Num())
-			{
-				LastIndex -= SyncIndex;
-			}
+			LastIndex = TransformLastIndex(LastIndex) + TransformLastIndex(SyncIndex);
 		}
 		else
 		{
 			checkf(InSyncTriggers.IsEmpty(), TEXT("SyncTriggers view should be the same size as the sample view or not supplied (size of 0)."));
-
-			float* OutSamples = OutIndicesView.GetData();
-			const float* SyncTrigData = InSyncTriggers.GetData();
 			for (int32 i = 0; i < OutIndicesView.Num(); ++i)
 			{
 				OutSamples[i] *= i;
 				OutSamples[i] += LastIndex;
 			}
-		}
 
-		if (InFreqModulator.Num() == OutIndicesView.Num())
-		{
-			NextIndex = (InFreqModulator.Last() * Settings.Freq * InTableView.Num()) + LastIndex;
-		}
-		else
-		{
-			checkf(InFreqModulator.IsEmpty(), TEXT("FreqModulator view should be the same size as the sample view or not supplied (size of 0)."));
-			NextIndex = (Settings.Freq * InTableView.Num()) + LastIndex;
+			LastIndex = TransformLastIndex(LastIndex);
 		}
 
-		LastIndex = FMath::Frac(NextIndex / InTableView.Num()) * InTableView.Num();
+		LastIndex = FMath::Frac(LastIndex / InTableView.Num()) * InTableView.Num();
 	}
 
 	void FWaveTableSampler::ComputeIndexPhase(TArrayView<const float> InTableView, TArrayView<const float> InPhaseModulator, TArrayView<float> OutIndicesView)
