MetasoundWaveTableEvaluateNode

diff --git a/Engine/Plugins/Runtime/Metasound/Source/MetasoundEngine/Private/MetasoundWaveTableEnvelopeNode.cpp b/Engine/Plugins/Runtime/Metasound/Source/MetasoundEngine/Private/MetasoundWaveTableEnvelopeNode.cpp
--- a/Engine/Plugins/Runtime/Metasound/Source/MetasoundEngine/Private/MetasoundWaveTableEnvelopeNode.cpp
+++ b/Engine/Plugins/Runtime/Metasound/Source/MetasoundEngine/Private/MetasoundWaveTableEnvelopeNode.cpp
@@ -21,12 +21,6 @@
 
 namespace Metasound
 {
-	DECLARE_METASOUND_ENUM(WaveTable::FWaveTableSampler::ESingleSampleMode, WaveTable::FWaveTableSampler::ESingleSampleMode::Zero, METASOUNDENGINE_API,
-	FEnumWaveTableEnvelopeMode, FEnumWaveTableEnvelopeModeTypeInfo, FEnumWaveTableEnvelopeModeReadRef, FEnumWaveTableEnvelopeModeWriteRef);
-
-	DECLARE_METASOUND_ENUM(WaveTable::FWaveTableSampler::EInterpolationMode, WaveTable::FWaveTableSampler::EInterpolationMode::Linear, METASOUNDENGINE_API,
-	FEnumWaveTableInterpolationMode, FEnumWaveTableInterpModeTypeInfo, FEnumWaveTableInterpModeReadRef, FEnumWaveTableInterpModeWriteRef);
-
 	DEFINE_METASOUND_ENUM_BEGIN(WaveTable::FWaveTableSampler::ESingleSampleMode, FEnumWaveTableEnvelopeMode, "WaveTableEnvelopeMode")
 		DEFINE_METASOUND_ENUM_ENTRY(WaveTable::FWaveTableSampler::ESingleSampleMode::Loop, "LoopDisplayName", "Loop", "EnvMode_LoopDescription", "Interpolates last value and first value in WaveTable, starting over interpolation of envelope on completion."),
 		DEFINE_METASOUND_ENUM_ENTRY(WaveTable::FWaveTableSampler::ESingleSampleMode::Hold, "HoldDisplayName", "Hold", "EnvMode_HoldDescription", "Holds last value in table if elapsed beyond WaveTable length"),
@@ -188,11 +182,9 @@ namespace Metasound
 		{
 			using namespace WaveTable;
 
-			float LastValue = *OutWriteRef;
-
 			if (bPaused)
 			{
-				return LastValue;
+				return *OutWriteRef;
 			}
 
 			switch (InSampleMode)
@@ -201,7 +193,7 @@ namespace Metasound
 				return 1.0f;
 
 				case FWaveTableSampler::ESingleSampleMode::Hold:
-				return LastValue;
+				return WaveTableReadRef->GetView().FinalValue;
 
 				case FWaveTableSampler::ESingleSampleMode::Zero:
 				case FWaveTableSampler::ESingleSampleMode::Loop:
@@ -268,9 +260,9 @@ namespace Metasound
 				}
 			}
 
-			TArrayView<const float> WaveTableView = WaveTableReadRef->GetView();
+			const WaveTable::FWaveTableView& WaveTableView = WaveTableReadRef->GetView();
 
-			if (WaveTableView.IsEmpty())
+			if (WaveTableView.SampleView.IsEmpty())
 			{
 				if (Elapsed >= 0.0f)
 				{
diff --git a/Engine/Plugins/Runtime/Metasound/Source/MetasoundEngine/Private/MetasoundWaveTableEvaluateNode.cpp b/Engine/Plugins/Runtime/Metasound/Source/MetasoundEngine/Private/MetasoundWaveTableEvaluateNode.cpp
new file mode 100644
--- /dev/null
+++ b/Engine/Plugins/Runtime/Metasound/Source/MetasoundEngine/Private/MetasoundWaveTableEvaluateNode.cpp
@@ -0,0 +1,188 @@
+// Copyright Epic Games, Inc. All Rights Reserved.
+
+#include "Interfaces/MetasoundFrontendSourceInterface.h"
+#include "MetasoundDataFactory.h"
+#include "MetasoundEngineNodesNames.h"
+#include "MetasoundEnumRegistrationMacro.h"
+#include "MetasoundExecutableOperator.h"
+#include "MetasoundFacade.h"
+#include "MetasoundNodeRegistrationMacro.h"
+#include "MetasoundPrimitives.h"
+#include "MetasoundParamHelper.h"
+#include "MetasoundStandardNodesCategories.h"
+#include "MetasoundWaveTable.h"
+#include "WaveTableSampler.h"
+
+#define LOCTEXT_NAMESPACE "MetasoundStandardNodes"
+
+namespace Metasound
+{
+	namespace WaveTableEvaluateNode
+	{
+		METASOUND_PARAM(WaveTableParam, "WaveTable", "The Wavetable to evaluate");
+		METASOUND_PARAM(InputParam, "Input", "[0, 1] the X input with which to evaluate the wavetable");
+		METASOUND_PARAM(OutParam, "Output", "The value of the wavetable at the specified point");
+	}
+	
+	class FMetasoundWaveTableEvaluateNodeOperator : public TExecutableOperator<FMetasoundWaveTableEvaluateNodeOperator>
+	{
+	public:
+		static const FVertexInterface& GetDefaultInterface()
+		{
+			using namespace WaveTable;
+			using namespace WaveTableEvaluateNode;
+
+			static const FVertexInterface DefaultInterface(
+				FInputVertexInterface(
+					TInputDataVertex<FWaveTable>(METASOUND_GET_PARAM_NAME_AND_METADATA(WaveTableParam)),
+					TInputDataVertex<float>(METASOUND_GET_PARAM_NAME_AND_METADATA(InputParam), 0.0f),
+					TInputDataVertex<FEnumWaveTableInterpolationMode>("Interpolation", FDataVertexMetadata
+					{
+						LOCTEXT("MetasoundWaveTableEvaluateNode_InterpDescription", "How the Evaluate interpolates between WaveTable values."),
+						LOCTEXT("MetasoundWaveTableEvaluateNode_Interp", "Interpolation"),
+						true /* bIsAdvancedDisplay */
+					}, static_cast<int32>(FWaveTableSampler::EInterpolationMode::Linear))
+				),
+				FOutputVertexInterface(
+					TOutputDataVertex<float>(METASOUND_GET_PARAM_NAME_AND_METADATA(OutParam))
+				)
+			);
+
+			return DefaultInterface;
+		}
+
+		static const FNodeClassMetadata& GetNodeInfo()
+		{
+			auto CreateNodeClassMetadata = []() -> FNodeClassMetadata
+			{
+				FNodeClassMetadata Metadata
+				{
+					{ EngineNodes::Namespace, "WaveTableEvaluate", "" },
+					1, // Major Version
+					0, // Minor Version
+					LOCTEXT("MetasoundWaveTableEvaluateNode_Name", "Evaluate WaveTable"),
+					LOCTEXT("MetasoundWaveTableEvaluateNode_Description", "Evaluates a wavetable for a given input value."),
+					PluginAuthor,
+					PluginNodeMissingPrompt,
+					GetDefaultInterface(),
+					{ NodeCategories::Envelopes },
+					{ METASOUND_LOCTEXT("WaveTableEvaluateEnvelopeKeyword", "Envelope"), METASOUND_LOCTEXT("WaveTableEvaluateCurveKeyword", "Curve")},
+					{ }
+				};
+
+				return Metadata;
+			};
+
+			static const FNodeClassMetadata Metadata = CreateNodeClassMetadata();
+			return Metadata;
+		}
+
+		static TUniquePtr<IOperator> CreateOperator(const FCreateOperatorParams& InParams, TArray<TUniquePtr<IOperatorBuildError>>& OutErrors)
+		{
+			using namespace WaveTable;
+
+			const FInputVertexInterface& InputInterface = InParams.Node.GetVertexInterface().GetInputInterface();
+			const FDataReferenceCollection& InputCollection = InParams.InputDataReferences;
+
+			FWaveTableReadRef InWaveTableReadRef = InputCollection.GetDataReadReferenceOrConstruct<FWaveTable>("WaveTable");
+			FFloatReadRef InInputReadRef = InputCollection.GetDataReadReferenceOrConstructWithVertexDefault<float>(InputInterface, "Input", InParams.OperatorSettings);
+			FEnumWaveTableInterpModeReadRef InInterpReadRef = InputCollection.GetDataReadReferenceOrConstructWithVertexDefault<FEnumWaveTableInterpolationMode>(InputInterface, "Interpolation", InParams.OperatorSettings);
+
+			return MakeUnique<FMetasoundWaveTableEvaluateNodeOperator>(InParams, InWaveTableReadRef, InInputReadRef, InInterpReadRef);
+		}
+
+		FMetasoundWaveTableEvaluateNodeOperator(
+			const FCreateOperatorParams& InParams,
+			const FWaveTableReadRef& InWaveTableReadRef,
+			const FFloatReadRef& InInputReadRef,
+			const FEnumWaveTableInterpModeReadRef& InInterpModeReadRef)
+			: WaveTableReadRef(InWaveTableReadRef)
+			, InputReadRef(InInputReadRef)
+			, InterpModeReadRef(InInterpModeReadRef)
+			, OutWriteRef(TDataWriteReferenceFactory<float>::CreateAny(InParams.OperatorSettings))
+		{
+			using namespace WaveTable;
+			FWaveTableSampler::FSettings Settings;
+			Settings.Freq = 0.0f; // Sampler phase is manually progressed via this node
+			Sampler = FWaveTableSampler(MoveTemp(Settings));
+		}
+
+		virtual ~FMetasoundWaveTableEvaluateNodeOperator() = default;
+
+		virtual void Bind(FVertexInterfaceData& InVertexData) const override
+		{
+			using namespace WaveTableEvaluateNode;
+			
+			FInputVertexInterfaceData& Inputs = InVertexData.GetInputs();
+			Inputs.BindReadVertex(METASOUND_GET_PARAM_NAME(WaveTableParam), WaveTableReadRef);
+			Inputs.BindReadVertex(METASOUND_GET_PARAM_NAME(InputParam), InputReadRef);
+			Inputs.BindReadVertex("Interpolation", InterpModeReadRef);
+
+			FOutputVertexInterfaceData& Outputs = InVertexData.GetOutputs();
+			Outputs.BindReadVertex(METASOUND_GET_PARAM_NAME(OutParam), OutWriteRef);
+		}
+		
+		virtual FDataReferenceCollection GetInputs() const override
+		{
+			// This should never be called. Bind(...) is called instead. This method
+			// exists as a stop-gap until the API can be deprecated and removed. 
+			checkNoEntry();
+
+			FDataReferenceCollection InputDataReferences;
+			return InputDataReferences;
+		}
+
+		virtual FDataReferenceCollection GetOutputs() const override
+		{
+			// This should never be called. Bind(...) is called instead. This method
+			// exists as a stop-gap until the API can be deprecated and removed. 
+			checkNoEntry();
+
+			FDataReferenceCollection OutputDataReferences;
+			return OutputDataReferences;
+		}
+
+		void Execute()
+		{
+			using namespace WaveTable;
+			const float Input = *InputReadRef;
+
+			float NextValue = 0.f;
+			
+			const FWaveTableView& WaveTableView = WaveTableReadRef->GetView();
+
+			constexpr auto SampleMode =  FWaveTableSampler::ESingleSampleMode::Hold;
+			
+			Sampler.Reset();
+			Sampler.SetInterpolationMode(*InterpModeReadRef);
+			Sampler.SetPhase(FMath::Clamp(Input, 0.0f, 1.0f));
+			Sampler.Process(WaveTableView, NextValue, SampleMode);
+
+			*OutWriteRef = NextValue;
+		}
+
+	private:
+		FWaveTableReadRef WaveTableReadRef;
+		FFloatReadRef InputReadRef;
+		FEnumWaveTableInterpModeReadRef InterpModeReadRef;
+
+		WaveTable::FWaveTableSampler Sampler;
+
+		FFloatWriteRef OutWriteRef;
+	};
+
+	class FMetasoundWaveTableEvaluateNode : public FNodeFacade
+	{
+	public:
+		FMetasoundWaveTableEvaluateNode(const FNodeInitData& InInitData)
+			: FNodeFacade(InInitData.InstanceName, InInitData.InstanceID, TFacadeOperatorClass<FMetasoundWaveTableEvaluateNodeOperator>())
+		{
+		}
+
+		virtual ~FMetasoundWaveTableEvaluateNode() = default;
+	};
+
+	METASOUND_REGISTER_NODE(FMetasoundWaveTableEvaluateNode)
+} // namespace Metasound
+
+#undef LOCTEXT_NAMESPACE // MetasoundStandardNodes
diff --git a/Engine/Plugins/Runtime/Metasound/Source/MetasoundEngine/Private/MetasoundWaveTableGetNode.cpp b/Engine/Plugins/Runtime/Metasound/Source/MetasoundEngine/Private/MetasoundWaveTableGetNode.cpp
--- a/Engine/Plugins/Runtime/Metasound/Source/MetasoundEngine/Private/MetasoundWaveTableGetNode.cpp
+++ b/Engine/Plugins/Runtime/Metasound/Source/MetasoundEngine/Private/MetasoundWaveTableGetNode.cpp
@@ -118,103 +118,90 @@ namespace Metasound
 			const FWaveTableBankAsset& WaveTableBankAsset = *WaveTableBankReadRef;
 
 			FWaveTableBankAssetProxyPtr Proxy = WaveTableBankAsset.GetProxy();
-			if (Proxy.IsValid())
+			if (Proxy.IsValid() == false)
 			{
-				const TArray<FWaveTable>& WaveTables = Proxy->GetWaveTables();
-				if (!WaveTables.IsEmpty())
+				LastTableIndex = -1.0f;
+				return;
+			}
+			
+			const TArray<FWaveTable>& WaveTables = Proxy->GetWaveTables();
+			if (WaveTables.IsEmpty())
+			{
+				LastTableIndex = -1.0f;
+				return;
+			}
+			
+			int32 NumSamples = WaveTables.Last().Num();
+			FWaveTable& OutputWaveTable = *OutTable;
+			OutputWaveTable.SetNum(NumSamples);
+			OutputWaveTable.Zero();
+
+			if (NumSamples <= 0)
+			{
+				LastTableIndex = -1.0f;
+				return;
+			}
+			
+			auto WrapIndex = [Max = WaveTables.Num()](float& InOutIndex)
+			{
+				InOutIndex = FMath::Abs(InOutIndex); // Avoids remainder offset flip at 0 crossing
+				const int32 WrapIndex = FMath::TruncToInt32(InOutIndex) % Max;
+				const float Remainder = FMath::Frac(InOutIndex);
+				InOutIndex = WrapIndex + Remainder;
+			};
+
+			float NextTableIndex = *TableIndexReadRef;
+			WrapIndex(NextTableIndex);
+
+			if (WaveTables.Num() == 1)
+			{
+				OutputWaveTable = WaveTables.Last();
+			}
+			else
+			{
+				if (LastTableIndex < 0.0f)
 				{
-					int32 NumSamples = WaveTables.Last().Num();
-					FWaveTable& OutputWaveTable = *OutTable;
-					OutputWaveTable.SetNum(NumSamples);
-					OutputWaveTable.Zero();
+					LastTableIndex = NextTableIndex;
+				}
+				else
+				{
+					WrapIndex(LastTableIndex);
+				}
+
+				const int32 StartHighIndex = FMath::CeilToInt32(LastTableIndex) % WaveTables.Num();
+				const int32 StartLowIndex = FMath::FloorToInt32(LastTableIndex);
+
+				const float StartHighGain = LastTableIndex - StartLowIndex;
+				const float StartLowGain = 1.0f - StartHighGain;
+
+				float FinalValue = 0.f;
 
-					if (NumSamples > 0)
+				for (int32 TableIndex = 0; TableIndex < WaveTables.Num(); ++TableIndex)
+				{
+					float Gain = 0.0f;
+
+					if (TableIndex == StartLowIndex)
 					{
-						auto WrapIndex = [Max = WaveTables.Num()](float& InOutIndex)
-						{
-							InOutIndex = FMath::Abs(InOutIndex); // Avoids remainder offset flip at 0 crossing
-							const int32 WrapIndex = FMath::TruncToInt32(InOutIndex) % Max;
-							const float Remainder = FMath::Frac(InOutIndex);
-							InOutIndex = WrapIndex + Remainder;
-						};
-
-						float NextTableIndex = *TableIndexReadRef;
-						WrapIndex(NextTableIndex);
-
-						if (WaveTables.Num() == 1)
-						{
-							OutputWaveTable = WaveTables.Last();
-						}
-						else
-						{
-							if (LastTableIndex < 0.0f)
-							{
-								LastTableIndex = NextTableIndex;
-							}
-							else
-							{
-								WrapIndex(LastTableIndex);
-							}
-
-							const int32 StartHighIndex = FMath::CeilToInt32(LastTableIndex) % WaveTables.Num();
-							const int32 StartLowIndex = FMath::FloorToInt32(LastTableIndex);
-
-							const float StartHighGain = LastTableIndex - StartLowIndex;
-							const float StartLowGain = 1.0f - StartHighGain;
-
-							const int32 EndHighIndex = FMath::CeilToInt32(NextTableIndex) % WaveTables.Num();
-							const int32 EndLowIndex = FMath::FloorToInt32(NextTableIndex);
-
-							const float EndHighGain = NextTableIndex - EndLowIndex;
-							const float EndLowGain = 1.0f - EndHighGain;
-
-							for (int32 TableIndex = 0; TableIndex < WaveTables.Num(); ++TableIndex)
-							{
-								float Gain = 0.0f;
-								float GainDelta = 0.0f;
-
-								if (TableIndex == StartLowIndex)
-								{
-									Gain = StartLowGain;
-								}
-								else if (TableIndex == StartHighIndex)
-								{
-									Gain = StartHighGain;
-								}
-
-								if (TableIndex == EndLowIndex)
-								{
-									GainDelta = (EndLowGain - Gain) / NumSamples;
-								}
-								else if (TableIndex == EndHighIndex)
-								{
-									GainDelta = (EndHighGain - Gain) / NumSamples;
-								}
-
-								if (Gain > 0.0f || GainDelta > 0.0f)
-								{
-									const FWaveTable& InputWaveTable = WaveTables[TableIndex];
-									Audio::ArrayMixIn(InputWaveTable.GetView(), OutputWaveTable.GetView(), Gain);
-								}
-							}
-						}
-
-						LastTableIndex = NextTableIndex;
+						Gain = StartLowGain;
 					}
-					else
+					else if (TableIndex == StartHighIndex)
 					{
-						LastTableIndex = -1.0f;
+						Gain = StartHighGain;
+					}
+
+					if (Gain > 0.0f)
+					{
+						const FWaveTableView& InputWaveTable = WaveTables[TableIndex].GetView();
+						
+						Audio::ArrayMixIn(InputWaveTable.SampleView, OutputWaveTable.GetSamples(), Gain);
+						FinalValue += InputWaveTable.FinalValue * Gain;
 					}
 				}
-				else
-				{
-					LastTableIndex = -1.0f;
-				}
-			}
-			else
-			{
-				LastTableIndex = -1.0f;
+				
+				OutputWaveTable.SetFinalValue(FinalValue);
 			}
+
+			LastTableIndex = NextTableIndex;
 		}
 
 	private:
diff --git a/Engine/Plugins/Runtime/Metasound/Source/MetasoundEngine/Public/MetasoundWaveTable.h b/Engine/Plugins/Runtime/Metasound/Source/MetasoundEngine/Public/MetasoundWaveTable.h
--- a/Engine/Plugins/Runtime/Metasound/Source/MetasoundEngine/Public/MetasoundWaveTable.h
+++ b/Engine/Plugins/Runtime/Metasound/Source/MetasoundEngine/Public/MetasoundWaveTable.h
@@ -6,10 +6,11 @@
 #include "IAudioProxyInitializer.h"
 #include "MetasoundDataReference.h"
 #include "MetasoundDataTypeRegistrationMacro.h"
+#include "MetasoundEnumRegistrationMacro.h"
 #include "MetasoundPrimitives.h"
 #include "WaveTable.h"
 #include "WaveTableBank.h"
-
+#include "WaveTableSampler.h"
 
 namespace Metasound
 {
@@ -86,5 +87,11 @@ namespace Metasound
 
 	DECLARE_METASOUND_DATA_REFERENCE_TYPES(WaveTable::FWaveTable, METASOUNDENGINE_API, FWaveTableTypeInfo, FWaveTableReadRef, FWaveTableWriteRef)
 	DECLARE_METASOUND_DATA_REFERENCE_TYPES(FWaveTableBankAsset, METASOUNDENGINE_API, FWaveTableBankAssetTypeInfo, FWaveTableBankAssetReadRef, FWaveTableBankAssetWriteRef)
+
+	DECLARE_METASOUND_ENUM(WaveTable::FWaveTableSampler::ESingleSampleMode, WaveTable::FWaveTableSampler::ESingleSampleMode::Zero, METASOUNDENGINE_API,
+	FEnumWaveTableEnvelopeMode, FEnumWaveTableEnvelopeModeTypeInfo, FEnumWaveTableEnvelopeModeReadRef, FEnumWaveTableEnvelopeModeWriteRef);
+
+	DECLARE_METASOUND_ENUM(WaveTable::FWaveTableSampler::EInterpolationMode, WaveTable::FWaveTableSampler::EInterpolationMode::Linear, METASOUNDENGINE_API,
+	FEnumWaveTableInterpolationMode, FEnumWaveTableInterpModeTypeInfo, FEnumWaveTableInterpModeReadRef, FEnumWaveTableInterpModeWriteRef);
 } // namespace Metasound
 
diff --git a/Engine/Plugins/Runtime/WaveTable/Source/WaveTable/Private/WaveTableBank.cpp b/Engine/Plugins/Runtime/WaveTable/Source/WaveTable/Private/WaveTableBank.cpp
--- a/Engine/Plugins/Runtime/WaveTable/Source/WaveTable/Private/WaveTableBank.cpp
+++ b/Engine/Plugins/Runtime/WaveTable/Source/WaveTable/Private/WaveTableBank.cpp
@@ -12,6 +12,34 @@ TSharedPtr<Audio::IProxyData> UWaveTableBank::CreateProxyData(const Audio::FProx
 	return MakeShared<FWaveTableBankAssetProxy>(*this);
 }
 
+void UWaveTableBank::Serialize(FArchive& Ar)
+{
+	if (Ar.IsLoading() || Ar.IsSaving())
+	{
+		for (FWaveTableBankEntry& Entry : Entries)
+		{
+			if (Entry.Transform.FinalValue == WaveTable::InvalidWaveTableValue)
+			{
+#if WITH_EDITOR
+				RefreshWaveTables();
+				return;
+#else
+				if(Entry.Transform.WaveTable.IsEmpty())
+				{
+					Entry.Transform.FinalValue = 0.f;
+				}
+				else
+				{
+					Entry.Transform.FinalValue = Entry.Transform.WaveTable.Last();
+				}
+#endif
+			}
+		}
+	}
+
+	UObject::Serialize(Ar);
+}
+
 #if WITH_EDITOR
 void UWaveTableBank::RefreshWaveTables()
 {
@@ -26,13 +54,14 @@ void UWaveTableBank::RefreshWaveTables()
 	WaveTableLengthSec = 0.0f;
 	for (FWaveTableBankEntry& Entry : Entries)
 	{
+		float& FinalValue = Entry.Transform.FinalValue;
 		TArray<float>& TransformWaveTable = Entry.Transform.WaveTable;
 		TransformWaveTable.Empty();
 
 		const int32 WaveTableSize = WaveTable::GetWaveTableSize(Resolution, Entry.Transform.Curve, MaxPCMSize);
 		TransformWaveTable.AddZeroed(WaveTableSize);
 
-		Entry.Transform.CreateWaveTable(TransformWaveTable, bBipolar);
+		Entry.Transform.CreateWaveTable(TransformWaveTable, FinalValue, bBipolar);
 		WaveTableSizeMB += sizeof(float) * WaveTableSize;
 		WaveTableLengthSec = FMath::Max(WaveTableLengthSec, WaveTableSize / 48000.f);
 	}
diff --git a/Engine/Plugins/Runtime/WaveTable/Source/WaveTable/Private/WaveTableImporter.cpp b/Engine/Plugins/Runtime/WaveTable/Source/WaveTable/Private/WaveTableImporter.cpp
--- a/Engine/Plugins/Runtime/WaveTable/Source/WaveTable/Private/WaveTableImporter.cpp
+++ b/Engine/Plugins/Runtime/WaveTable/Source/WaveTable/Private/WaveTableImporter.cpp
@@ -3,14 +3,9 @@
 
 #include "DSP/FloatArrayMath.h"
 #include "HAL/Platform.h"
-#include "Misc/AssertionMacros.h"
-#include "UObject/SoftObjectPath.h"
-#include "UObject/Object.h"
-#include "UObject/ObjectMacros.h"
-#include "UObject/Package.h"
+#include "WaveTable.h"
 #include "WaveTableSampler.h"
 
-
 namespace WaveTable
 {
 	FImporter::FImporter(const FWaveTableSettings& InSettings, bool bInBipolar)
@@ -24,7 +19,9 @@ namespace WaveTable
 	void FImporter::Process(TArray<float>& OutWaveTable)
 	{
 		// 1. Get editable section of source PCM
-		TArrayView<const float> EditSourceView = Settings.GetEditSourceView();
+		const TArrayView<const float> View = Settings.GetEditSourceView();
+		const FWaveTableView EditSourceView = FWaveTableView(View, View.Last());
+		
 		if (!EditSourceView.IsEmpty())
 		{
 			// 2. Resample into table
diff --git a/Engine/Plugins/Runtime/WaveTable/Source/WaveTable/Private/WaveTableSampler.cpp b/Engine/Plugins/Runtime/WaveTable/Source/WaveTable/Private/WaveTableSampler.cpp
--- a/Engine/Plugins/Runtime/WaveTable/Source/WaveTable/Private/WaveTableSampler.cpp
+++ b/Engine/Plugins/Runtime/WaveTable/Source/WaveTable/Private/WaveTableSampler.cpp
@@ -1,6 +1,7 @@
 // Copyright Epic Games, Inc. All Rights Reserved.
 #include "WaveTableSampler.h"
 
+#include "WaveTable.h"
 #include "DSP/Dsp.h"
 #include "DSP/FloatArrayMath.h"
 
@@ -181,8 +182,13 @@ namespace WaveTable
 		Settings.Phase = InPhase;
 	}
 
-	float FWaveTableSampler::Process(TArrayView<const float> InTableView, float& OutSample, ESingleSampleMode InMode)
+	float FWaveTableSampler::Process(const FWaveTableView& InTableView, float& OutSample, ESingleSampleMode InMode)
 	{
+		if(InTableView.IsEmpty())
+		{
+			return 0.f;
+		}
+		
 		TArrayView<float> OutSamplesView { &OutSample, 1 };
 		float Index = Process(InTableView, { }, { }, { }, OutSamplesView);
 
@@ -197,36 +203,27 @@ namespace WaveTable
 		// to see if interpolation needs to be re-evaluated.
 		// This check is expensive in the context of buffer
 		// eval and thus only supported in "single-sample" mode.
-		const bool bIsFinalInterp = Index > InTableView.Num() - 1;
+		const bool bIsFinalInterp = Index > InTableView.SampleView.Num() - 1;
 		const bool bHasCompleted = Index == 0 && (Settings.Phase > 0.0f || LastIndex > 0);
 		if (bIsFinalInterp || bHasCompleted)
 		{
+			float ModeValue = 0.f;
+			
 			switch (InMode)
 			{
-				case ESingleSampleMode::Hold:
+				case ESingleSampleMode::Zero:
+				case ESingleSampleMode::Unit:
 				{
-					OutSample = InTableView.Last();
-
+					ModeValue = static_cast<float>(InMode);
 				}
 				break;
 
-				case ESingleSampleMode::Zero:
-				case ESingleSampleMode::Unit:
+				case ESingleSampleMode::Hold:
 				{
-					const float ModeValue = static_cast<float>(InMode);
-					if (bIsFinalInterp)
-					{
-						TArray<float> FinalInterp = { InTableView.Last(), ModeValue };
-						OutSample = FMath::Frac(Index);
-						Interpolate(FinalInterp, OutSamplesView, Settings.InterpolationMode);
-					}
-					else
-					{
-						OutSample = ModeValue;
-					}
+					ModeValue = InTableView.FinalValue;
 				}
 				break;
-
+				
 				case ESingleSampleMode::Loop:
 				{
 					checkNoEntry(); // Should be handled above
@@ -237,6 +234,17 @@ namespace WaveTable
 				break;
 			}
 
+			if (bIsFinalInterp)
+			{
+				TArray<float> FinalInterp = { InTableView.SampleView.Last(), ModeValue };
+				OutSample = FMath::Frac(Index);
+				Interpolate(FinalInterp, OutSamplesView, Settings.InterpolationMode);
+			}
+			else
+			{
+				OutSample = ModeValue;
+			}
+
 			if (!FMath::IsNearlyEqual(Settings.Amplitude, 1.0f))
 			{
 				OutSample *= Settings.Amplitude;
@@ -251,18 +259,18 @@ namespace WaveTable
 		return Index;
 	}
 
-	float FWaveTableSampler::Process(TArrayView<const float> InTableView, TArrayView<float> OutSamplesView)
+	float FWaveTableSampler::Process(const FWaveTableView& InTableView, TArrayView<float> OutSamplesView)
 	{
 		return Process(InTableView, { }, { }, { }, OutSamplesView);
 	}
 
-	float FWaveTableSampler::Process(TArrayView<const float> InTableView, TArrayView<const float> InFreqModulator, TArrayView<const float> InPhaseModulator, TArrayView<const float> InSyncTriggers, TArrayView<float> OutSamplesView)
+	float FWaveTableSampler::Process(const FWaveTableView& InTableView, TArrayView<const float> InFreqModulator, TArrayView<const float> InPhaseModulator, TArrayView<const float> InSyncTriggers, TArrayView<float> OutSamplesView)
 	{
 		float CurViewLastIndex = 0.0f;
 
 		if (!OutSamplesView.IsEmpty())
 		{
-			if (InTableView.IsEmpty())
+			if (InTableView.SampleView.IsEmpty())
 			{
 				Audio::ArraySetToConstantInplace(OutSamplesView, 0.0f);
 			}
@@ -270,13 +278,13 @@ namespace WaveTable
 			{
 				// Compute index frequency & phase and stuff into samples
 				// view. Interpolator will convert to associated values.
-				ComputeIndexFrequency(InTableView, InFreqModulator, InSyncTriggers, OutSamplesView);
-				ComputeIndexPhase(InTableView, InPhaseModulator, OutSamplesView);
+				ComputeIndexFrequency(InTableView.SampleView, InFreqModulator, InSyncTriggers, OutSamplesView);
+				ComputeIndexPhase(InTableView.SampleView, InPhaseModulator, OutSamplesView);
 
 				// Capture last index and return in case caller is interested in progress/phase through table (ex. for enveloping).
 				CurViewLastIndex = OutSamplesView.Last();
 
-				Interpolate(InTableView, OutSamplesView, Settings.InterpolationMode);
+				Interpolate(InTableView.SampleView, OutSamplesView, Settings.InterpolationMode);
 			}
 
 			if (!FMath::IsNearlyEqual(Settings.Amplitude, 1.0f))
diff --git a/Engine/Plugins/Runtime/WaveTable/Source/WaveTable/Private/WaveTableTransform.cpp b/Engine/Plugins/Runtime/WaveTable/Source/WaveTable/Private/WaveTableTransform.cpp
--- a/Engine/Plugins/Runtime/WaveTable/Source/WaveTable/Private/WaveTableTransform.cpp
+++ b/Engine/Plugins/Runtime/WaveTable/Source/WaveTable/Private/WaveTableTransform.cpp
@@ -11,8 +11,9 @@ void FWaveTableTransform::Apply(float& InOutValue, bool bInBipolar) const
 {
 	if (WaveTable.IsEmpty())
 	{
+		float DummyVal = 0.f;
 		TArrayView<float> ValueView = { &InOutValue, 1 };
-		bInBipolar ? SampleCurveBipolar(ValueView) : SampleCurveUnipolar(ValueView);
+		bInBipolar ? SampleCurveBipolar(ValueView, DummyVal) : SampleCurveUnipolar(ValueView, DummyVal);
 	}
 	else
 	{
@@ -20,10 +21,9 @@ void FWaveTableTransform::Apply(float& InOutValue, bool bInBipolar) const
 		TArrayView<float> ValueView = { &InOutValue, 1 };
 		WaveTable::FWaveTableSampler::Interpolate(WaveTable, ValueView);
 	}
-
 }
 
-void FWaveTableTransform::SampleCurveBipolar(TArrayView<float> InOutTable) const
+void FWaveTableTransform::SampleCurveBipolar(TArrayView<float> InOutTable, float& OutFinalValue) const
 {
 	// Clamp the input
 	Audio::ArrayClampInPlace(InOutTable, 0.0f, 1.0f);
@@ -37,6 +37,8 @@ void FWaveTableTransform::SampleCurveBipolar(TArrayView<float> InOutTable) const
 			{
 				ValueData[i] = CurveCustom.Eval(ValueData[i]);
 			}
+				
+			OutFinalValue = CurveCustom.Eval(1.f);
 			break;
 		}
 
@@ -48,6 +50,8 @@ void FWaveTableTransform::SampleCurveBipolar(TArrayView<float> InOutTable) const
 				{
 					ValueData[i] = CurveShared->FloatCurve.Eval(ValueData[i]);
 				}
+				
+				OutFinalValue = CurveShared->FloatCurve.Eval(1.f);
 			}
 			break;
 		}
@@ -58,6 +62,8 @@ void FWaveTableTransform::SampleCurveBipolar(TArrayView<float> InOutTable) const
 			{
 				ValueData[i] = FMath::Lerp(1.0f, -1.0f, ValueData[i]);
 			}
+				
+			OutFinalValue = -1.f;
 		}
 		break;
 
@@ -67,6 +73,8 @@ void FWaveTableTransform::SampleCurveBipolar(TArrayView<float> InOutTable) const
 			{
 				ValueData[i] = FMath::Lerp(-1.0f, 1.0f, ValueData[i]);
 			}
+			
+			OutFinalValue = 1.f;
 		}
 		break;
 
@@ -79,6 +87,8 @@ void FWaveTableTransform::SampleCurveBipolar(TArrayView<float> InOutTable) const
 				ValueData[i] *= (FMath::Pow(10.0f, Scalar * (ValueData[i] - 1.0f)));
 				ValueData[i] -= 0.5f;
 			}
+
+			OutFinalValue = 1.f;
 		}
 		break;
 
@@ -91,6 +101,8 @@ void FWaveTableTransform::SampleCurveBipolar(TArrayView<float> InOutTable) const
 				ValueData[i] = ((ValueData[i] - 1.0f) * FMath::Pow(10.0f, -1.0f * Scalar * ValueData[i])) + 1.0f;
 				ValueData[i] -= 0.5f;
 			}
+
+			OutFinalValue = -1.f;
 		}
 		break;
 
@@ -100,6 +112,8 @@ void FWaveTableTransform::SampleCurveBipolar(TArrayView<float> InOutTable) const
 			{
 				ValueData[i] = (Scalar * FMath::LogX(10.0f, ValueData[i])) + 0.5f;
 			}
+				
+			OutFinalValue = 1.f;
 		}
 		break;
 
@@ -109,6 +123,8 @@ void FWaveTableTransform::SampleCurveBipolar(TArrayView<float> InOutTable) const
 			{
 				ValueData[i] = FMath::Sin(HALF_PI * ValueData[i]) - 0.5f;
 			}
+
+			OutFinalValue = 1.f;
 		}
 		break;
 
@@ -118,6 +134,8 @@ void FWaveTableTransform::SampleCurveBipolar(TArrayView<float> InOutTable) const
 			{
 				ValueData[i] = 0.5f * FMath::Sin(((PI * ValueData[i]) - HALF_PI));
 			}
+
+			OutFinalValue = 1.f;
 		}
 		break;
 
@@ -127,6 +145,8 @@ void FWaveTableTransform::SampleCurveBipolar(TArrayView<float> InOutTable) const
 			{
 				ValueData[i] = FMath::Sin(2.0f * PI * ValueData[i]);
 			}
+
+			OutFinalValue = 0.f;
 		}
 		break;
 
@@ -135,11 +155,13 @@ void FWaveTableTransform::SampleCurveBipolar(TArrayView<float> InOutTable) const
 			if (WaveTable.IsEmpty())
 			{
 				Audio::ArraySetToConstantInplace(InOutTable, 0.0f);
+				OutFinalValue = 0.f;
 			}
 			else
 			{
 				constexpr auto InterpMode = WaveTable::FWaveTableSampler::EInterpolationMode::MaxValue;
 				WaveTable::FWaveTableSampler::Interpolate(WaveTable, InOutTable, InterpMode);
+				OutFinalValue = WaveTable.Last();
 			}
 		}
 		break;
@@ -154,7 +176,7 @@ void FWaveTableTransform::SampleCurveBipolar(TArrayView<float> InOutTable) const
 	Audio::ArrayClampInPlace(InOutTable, -1.0f, 1.0f);
 }
 
-void FWaveTableTransform::SampleCurveUnipolar(TArrayView<float> InOutTable) const
+void FWaveTableTransform::SampleCurveUnipolar(TArrayView<float> InOutTable, float& OutFinalValue) const
 {
 	// Clamp the input
 	Audio::ArrayClampInPlace(InOutTable, 0.0f, 1.0f);
@@ -168,6 +190,8 @@ void FWaveTableTransform::SampleCurveUnipolar(TArrayView<float> InOutTable) cons
 			{
 				ValueData[i] = CurveCustom.Eval(ValueData[i]);
 			}
+				
+			OutFinalValue = CurveCustom.Eval(1.f);
 			break;
 		}
 
@@ -179,6 +203,7 @@ void FWaveTableTransform::SampleCurveUnipolar(TArrayView<float> InOutTable) cons
 				{
 					ValueData[i] = CurveShared->FloatCurve.Eval(ValueData[i]);
 				}
+				OutFinalValue = CurveShared->FloatCurve.Eval(1.f);
 			}
 			break;
 		}
@@ -189,12 +214,15 @@ void FWaveTableTransform::SampleCurveUnipolar(TArrayView<float> InOutTable) cons
 			{
 				ValueData[i] = 1.0f - InOutTable[i];
 			}
+				
+			OutFinalValue = 0.f;
 		}
 		break;
 
 		case EWaveTableCurve::Linear:
 		{
 			// Do nothing, as output values == clamped input values
+			OutFinalValue = 1.f;
 		}
 		break;
 	
@@ -206,6 +234,8 @@ void FWaveTableTransform::SampleCurveUnipolar(TArrayView<float> InOutTable) cons
 			{
 				ValueData[i] *= (FMath::Pow(10.0f, Scalar * (ValueData[i] - 1.0f)));
 			}
+
+			OutFinalValue = 1.f;
 		}
 		break;
 
@@ -217,6 +247,8 @@ void FWaveTableTransform::SampleCurveUnipolar(TArrayView<float> InOutTable) cons
 			{
 				ValueData[i] = ((ValueData[i] - 1.0f) * FMath::Pow(10.0f, -1.0f * Scalar * ValueData[i])) + 1.0f;
 			}
+
+			OutFinalValue = 0.f;
 		}
 		break;
 
@@ -226,6 +258,8 @@ void FWaveTableTransform::SampleCurveUnipolar(TArrayView<float> InOutTable) cons
 			{
 				ValueData[i] = (Scalar * FMath::LogX(10.0f, ValueData[i])) + 1.0f;
 			}
+
+			OutFinalValue = 1.f;
 		}
 		break;
 
@@ -235,6 +269,8 @@ void FWaveTableTransform::SampleCurveUnipolar(TArrayView<float> InOutTable) cons
 			{
 				ValueData[i] = FMath::Sin(HALF_PI * ValueData[i]);
 			}
+				
+			OutFinalValue = 1.f;
 		}
 		break;
 
@@ -244,6 +280,8 @@ void FWaveTableTransform::SampleCurveUnipolar(TArrayView<float> InOutTable) cons
 			{
 				ValueData[i] = 0.5f * FMath::Sin(((PI * ValueData[i]) - HALF_PI)) + 0.5f;
 			}
+
+			OutFinalValue = 1.f;
 		}
 		break;
 
@@ -253,6 +291,8 @@ void FWaveTableTransform::SampleCurveUnipolar(TArrayView<float> InOutTable) cons
 			{
 				ValueData[i] = 0.5f * FMath::Sin(2.0f * PI * ValueData[i]) + 0.5f;
 			}
+
+			OutFinalValue = 0.f;
 		}
 		break;
 
@@ -261,11 +301,13 @@ void FWaveTableTransform::SampleCurveUnipolar(TArrayView<float> InOutTable) cons
 			if (WaveTable.IsEmpty())
 			{
 				Audio::ArraySetToConstantInplace(InOutTable, 0.0f);
+				OutFinalValue = 0.f;
 			}
 			else
 			{
 				constexpr auto InterpMode = WaveTable::FWaveTableSampler::EInterpolationMode::MaxValue;
 				WaveTable::FWaveTableSampler::Interpolate(WaveTable, InOutTable, InterpMode);
+				OutFinalValue = WaveTable.Last();
 			}
 		}
 		break;
@@ -294,18 +336,19 @@ void FWaveTableTransform::CacheCurve()
 }
 
 #if WITH_EDITOR
-void FWaveTableTransform::CopyToWaveTable(TArrayView<float> InOutTable, bool bInBipolar) const
+void FWaveTableTransform::CopyToWaveTable(TArrayView<float> InOutTable, float& OutFinalValue, bool bInBipolar) const
 {
 	const float Delta = 1.0f / InOutTable.Num();
+	
 	for (int32 i = 0; i < InOutTable.Num(); ++i)
 	{
 		InOutTable[i] = i * Delta;
 	}
 
-	bInBipolar ? SampleCurveBipolar(InOutTable) : SampleCurveUnipolar(InOutTable);
+	bInBipolar ? SampleCurveBipolar(InOutTable, OutFinalValue) : SampleCurveUnipolar(InOutTable, OutFinalValue);
 }
 
-void FWaveTableTransform::CreateWaveTable(TArray<float>& InOutTable, bool bInBipolar) const
+void FWaveTableTransform::CreateWaveTable(TArray<float>& InOutTable, float& OutFinalValue, bool bInBipolar) const
 {
 	switch (Curve)
 	{
@@ -328,7 +371,7 @@ void FWaveTableTransform::CreateWaveTable(TArray<float>& InOutTable, bool bInBip
 		case EWaveTableCurve::Sin_Full:
 		default:
 		{
-			CopyToWaveTable(InOutTable, bInBipolar);
+			CopyToWaveTable(InOutTable, OutFinalValue, bInBipolar);
 		}
 		break;
 	}
diff --git a/Engine/Plugins/Runtime/WaveTable/Source/WaveTable/Public/WaveTable.h b/Engine/Plugins/Runtime/WaveTable/Source/WaveTable/Public/WaveTable.h
--- a/Engine/Plugins/Runtime/WaveTable/Source/WaveTable/Public/WaveTable.h
+++ b/Engine/Plugins/Runtime/WaveTable/Source/WaveTable/Public/WaveTable.h
@@ -8,21 +8,47 @@
 
 namespace WaveTable
 {
+	struct WAVETABLE_API FWaveTableView
+	{
+		TArrayView<const float> SampleView;
+		float FinalValue;
+
+		explicit FWaveTableView(const TArray<float>& InSamples, const float InFinalValue)
+			: SampleView(InSamples), FinalValue(InFinalValue) {}
+		
+		explicit FWaveTableView(const TArrayView<const float>& InSamples, const float InFinalValue)
+			: SampleView(InSamples), FinalValue(InFinalValue) {}
+
+		// returns the length of the SampleView array, not including the additional FinalValue
+		int32 Num() const
+		{
+			return SampleView.Num();
+		}
+
+		bool IsEmpty() const
+		{
+			return SampleView.IsEmpty();
+		}
+	};
+	
 	struct WAVETABLE_API FWaveTable
 	{
 	private:
 		TArray<float> Samples;
+		float FinalValue;
 
 	public:
 		FWaveTable() = default;
 
-		FWaveTable(const TArray<float>& InSamples)
+		FWaveTable(const TArray<float>& InSamples, const float InFinalValue)
 			: Samples(InSamples)
+			, FinalValue(InFinalValue)
 		{
 		}
 
-		FWaveTable(TArray<float>&& InSamples)
+		FWaveTable(TArray<float>&& InSamples, const float InFinalValue)
 			: Samples(MoveTemp(InSamples))
+			, FinalValue(InFinalValue)
 		{
 		}
 
@@ -36,12 +62,12 @@ namespace WaveTable
 			return Samples.GetData();
 		}
 
-		TArrayView<float> GetView()
+		FWaveTableView GetView() const
 		{
-			return Samples;
+			return FWaveTableView(Samples, FinalValue);
 		}
 
-		TArrayView<const float> GetView() const
+		TArrayView<float> GetSamples()
 		{
 			return Samples;
 		}
@@ -61,6 +87,11 @@ namespace WaveTable
 			Samples = InSamples;
 		}
 
+		void SetFinalValue(const float InValue)
+		{
+			FinalValue = InValue;
+		}
+
 		void SetNum(int32 InSize)
 		{
 			Samples.SetNum(InSize);
diff --git a/Engine/Plugins/Runtime/WaveTable/Source/WaveTable/Public/WaveTableBank.h b/Engine/Plugins/Runtime/WaveTable/Source/WaveTable/Public/WaveTableBank.h
--- a/Engine/Plugins/Runtime/WaveTable/Source/WaveTable/Public/WaveTableBank.h
+++ b/Engine/Plugins/Runtime/WaveTable/Source/WaveTable/Public/WaveTableBank.h
@@ -27,13 +27,13 @@ class WAVETABLE_API UWaveTableBank : public UObject, public IAudioProxyDataFacto
 	GENERATED_BODY()
 
 public:
-	// Number of sampled cached for each curve in the given bank.
+	// Number of samples cached for each curve in the given bank.
 	UPROPERTY(EditAnywhere, Category = Options)
 	EWaveTableResolution Resolution = EWaveTableResolution::Res_256;
 
-	// Determines if output from curve/wavetable are to be clamped between [-1.0f, 1.0f]
-	// (i.e. for waveform generation, oscillation, etc.) or between [0.0f, 1.0f]
-	// (i.e. for enveloping) when sampling curve/wavetable
+	// Determines if output from curve/wavetable are to be clamped between 
+	// [-1.0f, 1.0f] (i.e. for waveform generation, oscillation, etc.)
+	// or [0.0f, 1.0f] (i.e. for enveloping) when sampling curve/wavetable
 	UPROPERTY(EditAnywhere, Category = Options)
 	bool bBipolar = true;
 
@@ -54,6 +54,8 @@ public:
 	/* IAudioProxyDataFactory Implementation */
 	virtual TSharedPtr<Audio::IProxyData> CreateProxyData(const Audio::FProxyDataInitParams& InitParams) override;
 
+	virtual void Serialize(FArchive& Ar) override;
+
 #if WITH_EDITOR
 	void RefreshWaveTables();
 
@@ -77,7 +79,7 @@ public:
 	{
 		Algo::Transform(InWaveTableBank.Entries, WaveTables, [](const FWaveTableBankEntry& Entry)
 		{
-			return Entry.Transform.WaveTable;
+			return WaveTable::FWaveTable(Entry.Transform.WaveTable, Entry.Transform.FinalValue);
 		});
 	}
 
diff --git a/Engine/Plugins/Runtime/WaveTable/Source/WaveTable/Public/WaveTableSampler.h b/Engine/Plugins/Runtime/WaveTable/Source/WaveTable/Public/WaveTableSampler.h
--- a/Engine/Plugins/Runtime/WaveTable/Source/WaveTable/Public/WaveTableSampler.h
+++ b/Engine/Plugins/Runtime/WaveTable/Source/WaveTable/Public/WaveTableSampler.h
@@ -9,6 +9,8 @@
 
 namespace WaveTable
 {
+	struct FWaveTableView;
+
 	class WAVETABLE_API FWaveTableSampler
 	{
 	public:
@@ -67,9 +69,9 @@ namespace WaveTable
 		// Interpolates and converts values in the given table for each provided index in the index-to-samples TArrayView (an array of sub-sample, floating point, indices)
 		static void Interpolate(TArrayView<const float> InTableView, TArrayView<float> InOutIndexToSamplesView, EInterpolationMode InterpMode = EInterpolationMode::Linear);
 
-		float Process(TArrayView<const float> InTableView, float& OutSample, ESingleSampleMode InMode = ESingleSampleMode::Zero);
-		float Process(TArrayView<const float> InTableView, TArrayView<float> OutSamplesView);
-		float Process(TArrayView<const float> InTableView, TArrayView<const float> InFreqModulator, TArrayView<const float> InPhaseModulator, TArrayView<const float> InSyncTriggers, TArrayView<float> OutSamplesView);
+		float Process(const FWaveTableView& InTableView, float& OutSample, ESingleSampleMode InMode = ESingleSampleMode::Zero);
+		float Process(const FWaveTableView& InTableView, TArrayView<float> OutSamplesView);
+		float Process(const FWaveTableView& InTableView, TArrayView<const float> InFreqModulator, TArrayView<const float> InPhaseModulator, TArrayView<const float> InSyncTriggers, TArrayView<float> OutSamplesView);
 
 		void Reset();
 
diff --git a/Engine/Plugins/Runtime/WaveTable/Source/WaveTable/Public/WaveTableTransform.h b/Engine/Plugins/Runtime/WaveTable/Source/WaveTable/Public/WaveTableTransform.h
--- a/Engine/Plugins/Runtime/WaveTable/Source/WaveTable/Public/WaveTableTransform.h
+++ b/Engine/Plugins/Runtime/WaveTable/Source/WaveTable/Public/WaveTableTransform.h
@@ -5,15 +5,15 @@
 #include "Containers/ArrayView.h"
 #include "Curves/CurveFloat.h"
 #include "Curves/RichCurve.h"
-#include "HAL/Platform.h"
-#include "Templates/Function.h"
-#include "UObject/SoftObjectPath.h"
-#include "UObject/Object.h"
-#include "UObject/ObjectMacros.h"
 #include "WaveTableSettings.h"
 
 #include "WaveTableTransform.generated.h"
 
+// use a value outside the [-1,1] range to detect uninitialized values
+namespace WaveTable
+{
+	constexpr float InvalidWaveTableValue = TNumericLimits<float>::Min();
+}
 
 USTRUCT(BlueprintType)
 struct WAVETABLE_API FWaveTableTransform
@@ -43,6 +43,9 @@ struct WAVETABLE_API FWaveTableTransform
 	UPROPERTY()
 	TArray<float> WaveTable;
 
+	UPROPERTY()
+	float FinalValue = WaveTable::InvalidWaveTableValue;
+
 #if WITH_EDITORONLY_DATA
 	UPROPERTY(EditAnywhere, Category = WaveTable)
 	FWaveTableSettings WaveTableSettings;
@@ -56,11 +59,11 @@ struct WAVETABLE_API FWaveTableTransform
 	// size of the resulting WaveTable is the provided ArrayView's size. If curve is itself
 	// set to 'File' and internally cached as WaveTable, will re-sample using max value
 	// interpolation. Does *not* apply destructive edits for using WaveTableSettings PCM data.
-	void CopyToWaveTable(TArrayView<float> InOutTable, bool bInBipolar) const;
+	void CopyToWaveTable(TArrayView<float> InOutTable, float& OutFinalValue, bool bInBipolar) const;
 
 	// Caches WaveTable for curve using the provided array length as resolution.
 	// Applies destructive edits as provided by WaveTableSettings if curve set to 'File'.
-	void CreateWaveTable(TArray<float>& InOutTable, bool bInBipolar) const;
+	void CreateWaveTable(TArray<float>& InOutTable, float& OutFinalValue, bool bInBipolar) const;
 #endif // WITH_EDITOR
 
 	/** Caches curve data.  Should a shared curve be selected, curve points are copied locally
@@ -72,8 +75,8 @@ struct WAVETABLE_API FWaveTableTransform
 
 private:
 	/** Clamps & applies transform to provided values as bipolar signal*/
-	void SampleCurveBipolar(TArrayView<float> InOutValues) const;
+	void SampleCurveBipolar(TArrayView<float> InOutValues, float& OutFinalValue) const;
 
 	/** Clamps & applies transform to provided values as unipolar signal */
-	void SampleCurveUnipolar(TArrayView<float> InOutValues) const;
+	void SampleCurveUnipolar(TArrayView<float> InOutValues, float& OutFinalValue) const;
 };
diff --git a/Engine/Plugins/Runtime/WaveTable/Source/WaveTableEditor/Private/WaveTableBankEditor.cpp b/Engine/Plugins/Runtime/WaveTable/Source/WaveTableEditor/Private/WaveTableBankEditor.cpp
--- a/Engine/Plugins/Runtime/WaveTable/Source/WaveTableEditor/Private/WaveTableBankEditor.cpp
+++ b/Engine/Plugins/Runtime/WaveTable/Source/WaveTableEditor/Private/WaveTableBankEditor.cpp
@@ -198,16 +198,18 @@ namespace WaveTable
 					TArray<float> KeyTable;
 					KeyTable.AddZeroed(PointCount);
 
+					float FinalValue = 0.f;
+
 					// Optimization for really big source files if set to 'File'. Do not recreate WaveTable on the fly
 					// for large samples as caching mechanism is too slow for rapid regeneration (i.e. when making
 					// interactive edits).
 					if (Transform->WaveTableSettings.SourcePCMData.Num() < BankEditorPrivate::MaxLiveEditSamples)
 					{
-						Transform->CreateWaveTable(KeyTable, bIsBipolar);
+						Transform->CreateWaveTable(KeyTable, FinalValue, bIsBipolar);
 					}
 					else
 					{
-						Transform->CopyToWaveTable(KeyTable, bIsBipolar);
+						Transform->CopyToWaveTable(KeyTable, FinalValue, bIsBipolar);
 					}
 
 					const float Delta = 1.0 / KeyTable.Num();
