fix crash in wave table importer when sampling an empty dataset

#rnx

diff --git a/Engine/Plugins/Runtime/WaveTable/Source/WaveTable/Private/WaveTableImporter.cpp b/Engine/Plugins/Runtime/WaveTable/Source/WaveTable/Private/WaveTableImporter.cpp
--- a/Engine/Plugins/Runtime/WaveTable/Source/WaveTable/Private/WaveTableImporter.cpp
+++ b/Engine/Plugins/Runtime/WaveTable/Source/WaveTable/Private/WaveTableImporter.cpp
@@ -19,59 +19,62 @@ namespace WaveTable
 	{
 		// 1. Get editable section of source PCM
 		const TArrayView<const float> View = Settings.GetEditSourceView();
+
+		if(View.IsEmpty())
+		{
+			return;
+		}
+		
 		const FWaveTableView EditSourceView = FWaveTableView(View, View.Last());
 		
-		if (!EditSourceView.IsEmpty())
+		// 2. Resample into table
+		Sampler.Reset();
+		Sampler.Process(EditSourceView, OutWaveTable);
+
+		// 3. Apply offset
+		float TableOffset = 0.0f;
+		if (Settings.bRemoveOffset || !bBipolar)
 		{
-			// 2. Resample into table
-			Sampler.Reset();
-			Sampler.Process(EditSourceView, OutWaveTable);
+			TableOffset = Audio::ArrayGetAverageValue(OutWaveTable);
+			Audio::ArrayAddConstantInplace(OutWaveTable, -1.0f * TableOffset);
+		}
 
-			// 3. Apply offset
-			float TableOffset = 0.0f;
-			if (Settings.bRemoveOffset || !bBipolar)
+		// 4. Normalize
+		if (Settings.bNormalize)
+		{
+			const float MaxValue = Audio::ArrayMaxAbsValue(OutWaveTable);
+			if (MaxValue > 0.0f)
 			{
-				TableOffset = Audio::ArrayGetAverageValue(OutWaveTable);
-				Audio::ArrayAddConstantInplace(OutWaveTable, -1.0f * TableOffset);
+				Audio::ArrayMultiplyByConstantInPlace(OutWaveTable, 1.0f / MaxValue);
 			}
+		}
 
-			// 4. Normalize
-			if (Settings.bNormalize)
-			{
-				const float MaxValue = Audio::ArrayMaxAbsValue(OutWaveTable);
-				if (MaxValue > 0.0f)
-				{
-					Audio::ArrayMultiplyByConstantInPlace(OutWaveTable, 1.0f / MaxValue);
-				}
-			}
+		// 5. Apply fades
+		if (Settings.FadeIn > 0.0f)
+		{
+			const int32 FadeLength = RatioToIndex(Settings.FadeIn, OutWaveTable.Num());
+			TArrayView<float> FadeView(OutWaveTable.GetData(), FadeLength);
+			Audio::ArrayFade(FadeView, 0.0f, 1.0f);
+		}
 
-			// 5. Apply fades
-			if (Settings.FadeIn > 0.0f)
-			{
-				const int32 FadeLength = RatioToIndex(Settings.FadeIn, OutWaveTable.Num());
-				TArrayView<float> FadeView(OutWaveTable.GetData(), FadeLength);
-				Audio::ArrayFade(FadeView, 0.0f, 1.0f);
-			}
+		if (Settings.FadeOut > 0.0f)
+		{
+			const int32 FadeLastIndex = RatioToIndex(Settings.FadeOut, OutWaveTable.Num());
+			const int32 FadeInitIndex = RatioToIndex(1.0f - Settings.FadeOut, OutWaveTable.Num());
+			TArrayView<float> FadeView = TArrayView<float>(&OutWaveTable[FadeInitIndex], FadeLastIndex + 1);
+			Audio::ArrayFade(FadeView, 1.0f, 0.0f);
+		}
 
-			if (Settings.FadeOut > 0.0f)
-			{
-				const int32 FadeLastIndex = RatioToIndex(Settings.FadeOut, OutWaveTable.Num());
-				const int32 FadeInitIndex = RatioToIndex(1.0f - Settings.FadeOut, OutWaveTable.Num());
-				TArrayView<float> FadeView = TArrayView<float>(&OutWaveTable[FadeInitIndex], FadeLastIndex + 1);
-				Audio::ArrayFade(FadeView, 1.0f, 0.0f);
-			}
+		// 6. Finalize if unipolar source
+		if (!bBipolar)
+		{
+			Audio::ArrayAbsInPlace(OutWaveTable);
 
-			// 6. Finalize if unipolar source
-			if (!bBipolar)
+			// If not requesting offset removal, add back (is always removed above for
+			// unipolar to ensure fades are centered around offset)
+			if (!Settings.bRemoveOffset)
 			{
-				Audio::ArrayAbsInPlace(OutWaveTable);
-
-				// If not requesting offset removal, add back (is always removed above for
-				// unipolar to ensure fades are centered around offset)
-				if (!Settings.bRemoveOffset)
-				{
-					Audio::ArrayAddConstantInplace(OutWaveTable, TableOffset);
-				}
+				Audio::ArrayAddConstantInplace(OutWaveTable, TableOffset);
 			}
 		}
 	} // namespace Importer
