Fix for crash in WaveTable editor
#rnx

diff --git a/Engine/Plugins/Runtime/WaveTable/Source/WaveTable/Private/WaveTableBank.cpp b/Engine/Plugins/Runtime/WaveTable/Source/WaveTable/Private/WaveTableBank.cpp
--- a/Engine/Plugins/Runtime/WaveTable/Source/WaveTable/Private/WaveTableBank.cpp
+++ b/Engine/Plugins/Runtime/WaveTable/Source/WaveTable/Private/WaveTableBank.cpp
@@ -13,34 +13,6 @@ TSharedPtr<Audio::IProxyData> UWaveTableBank::CreateProxyData(const Audio::FProx
 	return MakeShared<FWaveTableBankAssetProxy>(*this);
 }
 
-void UWaveTableBank::Serialize(FArchive& Ar)
-{
-	if (Ar.IsLoading() || Ar.IsSaving())
-	{
-		for (FWaveTableBankEntry& Entry : Entries)
-		{
-			if (Entry.Transform.FinalValue == WaveTable::InvalidWaveTableValue)
-			{
-#if WITH_EDITOR
-				RefreshWaveTables();
-				return;
-#else
-				if(Entry.Transform.WaveTable.IsEmpty())
-				{
-					Entry.Transform.FinalValue = 0.f;
-				}
-				else
-				{
-					Entry.Transform.FinalValue = Entry.Transform.WaveTable.Last();
-				}
-#endif
-			}
-		}
-	}
-
-	UObject::Serialize(Ar);
-}
-
 #if WITH_EDITOR
 void UWaveTableBank::RefreshWaveTables()
 {
@@ -55,14 +27,16 @@ void UWaveTableBank::RefreshWaveTables()
 	WaveTableLengthSec = 0.0f;
 	for (FWaveTableBankEntry& Entry : Entries)
 	{
-		float& FinalValue = Entry.Transform.FinalValue;
 		TArray<float>& TransformWaveTable = Entry.Transform.WaveTable;
 		TransformWaveTable.Empty();
 
 		const int32 WaveTableSize = WaveTable::GetWaveTableSize(Resolution, Entry.Transform.Curve, MaxPCMSize);
 		TransformWaveTable.AddZeroed(WaveTableSize);
 
+		float FinalValue = Entry.Transform.GetFinalValue();
 		Entry.Transform.CreateWaveTable(TransformWaveTable, FinalValue, bBipolar);
+		Entry.Transform.SetFinalValue(FinalValue);
+
 		WaveTableSizeMB += sizeof(float) * WaveTableSize;
 		WaveTableLengthSec = FMath::Max(WaveTableLengthSec, WaveTableSize / 48000.f);
 	}
diff --git a/Engine/Plugins/Runtime/WaveTable/Source/WaveTable/Private/WaveTableTransform.cpp b/Engine/Plugins/Runtime/WaveTable/Source/WaveTable/Private/WaveTableTransform.cpp
--- a/Engine/Plugins/Runtime/WaveTable/Source/WaveTable/Private/WaveTableTransform.cpp
+++ b/Engine/Plugins/Runtime/WaveTable/Source/WaveTable/Private/WaveTableTransform.cpp
@@ -335,6 +335,26 @@ void FWaveTableTransform::CacheCurve()
 	CurveShared = nullptr;
 }
 
+float FWaveTableTransform::GetFinalValue() const
+{
+	if (FMath::IsNearlyEqual(FinalValue, WaveTable::InvalidWaveTableValue))
+	{
+		if (WaveTable.IsEmpty())
+		{
+			return 0.f;
+		}
+
+		return WaveTable.Last();
+	}
+
+	return FinalValue;
+}
+
+void FWaveTableTransform::SetFinalValue(float InValue)
+{
+	FinalValue = InValue;
+}
+
 #if WITH_EDITOR
 void FWaveTableTransform::CopyToWaveTable(TArrayView<float> InOutTable, float& OutFinalValue, bool bInBipolar) const
 {
diff --git a/Engine/Plugins/Runtime/WaveTable/Source/WaveTable/Public/WaveTableBank.h b/Engine/Plugins/Runtime/WaveTable/Source/WaveTable/Public/WaveTableBank.h
--- a/Engine/Plugins/Runtime/WaveTable/Source/WaveTable/Public/WaveTableBank.h
+++ b/Engine/Plugins/Runtime/WaveTable/Source/WaveTable/Public/WaveTableBank.h
@@ -52,8 +52,6 @@ public:
 	/* IAudioProxyDataFactory Implementation */
 	virtual TSharedPtr<Audio::IProxyData> CreateProxyData(const Audio::FProxyDataInitParams& InitParams) override;
 
-	virtual void Serialize(FArchive& Ar) override;
-
 #if WITH_EDITOR
 	void RefreshWaveTables();
 
@@ -77,7 +75,7 @@ public:
 	{
 		Algo::Transform(InWaveTableBank.Entries, WaveTables, [](const FWaveTableBankEntry& Entry)
 		{
-			return WaveTable::FWaveTable(Entry.Transform.WaveTable, Entry.Transform.FinalValue);
+			return WaveTable::FWaveTable(Entry.Transform.WaveTable, Entry.Transform.GetFinalValue());
 		});
 	}
 
diff --git a/Engine/Plugins/Runtime/WaveTable/Source/WaveTable/Public/WaveTableTransform.h b/Engine/Plugins/Runtime/WaveTable/Source/WaveTable/Public/WaveTableTransform.h
--- a/Engine/Plugins/Runtime/WaveTable/Source/WaveTable/Public/WaveTableTransform.h
+++ b/Engine/Plugins/Runtime/WaveTable/Source/WaveTable/Public/WaveTableTransform.h
@@ -42,9 +42,6 @@ struct WAVETABLE_API FWaveTableTransform
 	UPROPERTY()
 	TArray<float> WaveTable;
 
-	UPROPERTY()
-	float FinalValue = WaveTable::InvalidWaveTableValue;
-
 #if WITH_EDITORONLY_DATA
 	UPROPERTY(EditAnywhere, Category = WaveTable)
 	FWaveTableSettings WaveTableSettings;
@@ -72,14 +69,19 @@ struct WAVETABLE_API FWaveTableTransform
 	  */
 	void CacheCurve();
 
+	// Returns the final value of the WaveTable
+	float GetFinalValue() const;
+
+	// Set the final value of the WaveTable
+	void SetFinalValue(float InValue);
+
 private:
+	UPROPERTY()
+	float FinalValue = WaveTable::InvalidWaveTableValue;
+
 	/** Clamps & applies transform to provided values as bipolar signal*/
 	void SampleCurveBipolar(TArrayView<float> InOutValues, float& OutFinalValue) const;
 
 	/** Clamps & applies transform to provided values as unipolar signal */
 	void SampleCurveUnipolar(TArrayView<float> InOutValues, float& OutFinalValue) const;
 };
-
-#if UE_ENABLE_INCLUDE_ORDER_DEPRECATED_IN_5_2
-#include "Curves/CurveFloat.h"
-#endif
