Adding support for AIFF files with the .aiff extension  on import

diff --git a/Engine/Plugins/Runtime/WaveTable/Source/WaveTable/Public/WaveTableSettings.h b/Engine/Plugins/Runtime/WaveTable/Source/WaveTable/Public/WaveTableSettings.h
--- a/Engine/Plugins/Runtime/WaveTable/Source/WaveTable/Public/WaveTableSettings.h
+++ b/Engine/Plugins/Runtime/WaveTable/Source/WaveTable/Public/WaveTableSettings.h
@@ -51,7 +51,7 @@ struct WAVETABLE_API FWaveTableSettings
 	GENERATED_USTRUCT_BODY()
 
 	// File to import
-	UPROPERTY(EditAnywhere, Category = Options, meta = (FilePathFilter = "Audio Files (*.aif, *.flac, *.ogg, *.wav)|*.aif;*.flac;*.ogg;*.wav"))
+	UPROPERTY(EditAnywhere, Category = Options, meta = (FilePathFilter = "Audio Files (*.aif, *.aiff, *.flac, *.ogg, *.wav)|*.aif;*.aiff;*.flac;*.ogg;*.wav"))
 	FFilePath FilePath;
 
 	// Index of channel in file to build WaveTable from (wraps if channel is greater than number in file)
diff --git a/Engine/Source/Editor/AudioEditor/Private/Factories/ReimportSoundFactory.cpp b/Engine/Source/Editor/AudioEditor/Private/Factories/ReimportSoundFactory.cpp
--- a/Engine/Source/Editor/AudioEditor/Private/Factories/ReimportSoundFactory.cpp
+++ b/Engine/Source/Editor/AudioEditor/Private/Factories/ReimportSoundFactory.cpp
@@ -17,6 +17,7 @@ UReimportSoundFactory::UReimportSoundFactory(const FObjectInitializer& ObjectIni
 
 #if WITH_SNDFILE_IO
 	Formats.Add(TEXT("aif;Audio Interchange File"));
+	Formats.Add(TEXT("aiff;Audio Interchange File Format"));
 	Formats.Add(TEXT("ogg;OGG Vorbis bitstream format"));
 	Formats.Add(TEXT("flac;Free Lossless Audio Codec"));
 #endif // WITH_SNDFILE_IO
diff --git a/Engine/Source/Editor/AudioEditor/Private/Factories/SoundFactory.cpp b/Engine/Source/Editor/AudioEditor/Private/Factories/SoundFactory.cpp
--- a/Engine/Source/Editor/AudioEditor/Private/Factories/SoundFactory.cpp
+++ b/Engine/Source/Editor/AudioEditor/Private/Factories/SoundFactory.cpp
@@ -148,6 +148,7 @@ USoundFactory::USoundFactory(const FObjectInitializer& ObjectInitializer)
 
 #if WITH_SNDFILE_IO
 	Formats.Add(TEXT("aif;Audio Interchange File"));
+	Formats.Add(TEXT("aiff;Audio Interchange File Format"));
 	Formats.Add(TEXT("ogg;OGG Vorbis bitstream format "));
 	Formats.Add(TEXT("flac;Free Lossless Audio Codec"));
 #endif // WITH_SNDFILE_IO
