Fixed mismatched copy constructor/assignment operator in FWaveTableBankAssetProxy.

diff --git a/Engine/Plugins/Runtime/WaveTable/Source/WaveTable/Public/WaveTableBank.h b/Engine/Plugins/Runtime/WaveTable/Source/WaveTable/Public/WaveTableBank.h
--- a/Engine/Plugins/Runtime/WaveTable/Source/WaveTable/Public/WaveTableBank.h
+++ b/Engine/Plugins/Runtime/WaveTable/Source/WaveTable/Public/WaveTableBank.h
@@ -66,11 +66,6 @@ class WAVETABLE_API FWaveTableBankAssetProxy : public Audio::TProxyData<FWaveTab
 public:
 	IMPL_AUDIOPROXY_CLASS(FWaveTableBankAssetProxy);
 
-	FWaveTableBankAssetProxy(const FWaveTableBankAssetProxy& InAssetProxy)
-		: WaveTables(InAssetProxy.WaveTables)
-	{
-	}
-
 	FWaveTableBankAssetProxy(const UWaveTableBank& InWaveTableBank)
 	{
 		Algo::Transform(InWaveTableBank.Entries, WaveTables, [](const FWaveTableBankEntry& Entry)
