MetaSound node support for traditional "sampler" usage of WaveTables
WaveTable memory reduction, better format support, player node
- Add support for 16bit assets & perform runtime BDC where necessary
- Add support for FixedSampleRate vs resolution
- Cache off and reuse single proxy
- Move entry data to table data in non-editor builds (instead of copying)

diff --git a/Engine/Plugins/Runtime/AudioModulation/Source/AudioModulation/Private/SoundModulationPatch.cpp b/Engine/Plugins/Runtime/AudioModulation/Source/AudioModulation/Private/SoundModulationPatch.cpp
--- a/Engine/Plugins/Runtime/AudioModulation/Source/AudioModulation/Private/SoundModulationPatch.cpp
+++ b/Engine/Plugins/Runtime/AudioModulation/Source/AudioModulation/Private/SoundModulationPatch.cpp
@@ -60,6 +60,21 @@ void USoundModulationPatch::PostEditChangeChainProperty(FPropertyChangedChainEve
 	Super::PostEditChangeChainProperty(InPropertyChangedEvent);
 }
 
+#if WITH_EDITORONLY_DATA
+void USoundModulationPatch::Serialize(FArchive& Ar)
+{
+	if (Ar.IsLoading())
+	{
+		for (FSoundControlModulationInput& Input : PatchSettings.Inputs)
+		{
+			Input.Transform.VersionTableData();
+		}
+	}
+
+	Super::Serialize(Ar);
+}
+#endif // WITH_EDITORONLY_DATA
+
 void USoundModulationPatch::PreSave(FObjectPreSaveContext InSaveContext)
 {
 	Super::PreSave(InSaveContext);
diff --git a/Engine/Plugins/Runtime/AudioModulation/Source/AudioModulation/Public/SoundModulationPatch.h b/Engine/Plugins/Runtime/AudioModulation/Source/AudioModulation/Public/SoundModulationPatch.h
--- a/Engine/Plugins/Runtime/AudioModulation/Source/AudioModulation/Public/SoundModulationPatch.h
+++ b/Engine/Plugins/Runtime/AudioModulation/Source/AudioModulation/Public/SoundModulationPatch.h
@@ -14,11 +14,13 @@
 class USoundControlBus;
 
 
+PRAGMA_DISABLE_DEPRECATION_WARNINGS
 USTRUCT(BlueprintType)
 struct AUDIOMODULATION_API FSoundModulationTransform : public FWaveTableTransform
 {
 	GENERATED_USTRUCT_BODY()
 };
+PRAGMA_ENABLE_DEPRECATION_WARNINGS
 
 
 USTRUCT(BlueprintType)
@@ -76,6 +78,9 @@ public:
 
 	virtual TUniquePtr<Audio::IModulatorSettings> CreateProxySettings() const override;
 
+#if WITH_EDITORONLY_DATA
+	virtual void Serialize(FArchive& Ar) override;
+#endif // WITH_EDITORONLY_DATA
 
 #if WITH_EDITOR
 	virtual void PreSave(FObjectPreSaveContext InSaveContext) override;
diff --git a/Engine/Plugins/Runtime/Metasound/Source/MetasoundEngine/Private/MetasoundWaveTableGetNode.cpp b/Engine/Plugins/Runtime/Metasound/Source/MetasoundEngine/Private/MetasoundWaveTableGetNode.cpp
--- a/Engine/Plugins/Runtime/Metasound/Source/MetasoundEngine/Private/MetasoundWaveTableGetNode.cpp
+++ b/Engine/Plugins/Runtime/Metasound/Source/MetasoundEngine/Private/MetasoundWaveTableGetNode.cpp
@@ -45,9 +45,14 @@ namespace Metasound
 				{
 					FNodeClassName { EngineNodes::Namespace, "MetasoundWaveTableGet", "" },
 					1, // Major Version
-					0, // Minor Version
+					1, // Minor Version
 					LOCTEXT("MetasoundWaveTableGetNode_Name", "Get WaveTable From Bank"),
-					LOCTEXT("MetasoundWaveTableGetNode_Description", "Gets or generates interpolated WaveTable from provided WaveTableBank asset."),
+					LOCTEXT("MetasoundWaveTableGetNode_Description",
+						"Gets or generates interpolated WaveTable from provided WaveTableBank asset based on asset sampling mode "
+						"(Table Interpolation is supported only for 'FixedResolution' banks).\n"
+						"v1.1: Now supports discrete selection of WaveTableBank assets set to 'Fixed"
+						"Sample Rate' and/or a bit depth of 16 bit, however index interpolation"
+						"is only supported for assets set to 'Fixed Resolution'"),
 					PluginAuthor,
 					PluginNodeMissingPrompt,
 					GetDefaultInterface(),
@@ -110,6 +115,69 @@ namespace Metasound
 			return Outputs;
 		}
 
+	private:
+		FORCEINLINE static void WrapIndex(int32 InMax, float& InOutIndex)
+		{
+			InOutIndex = FMath::Abs(InOutIndex); // Avoids remainder offset flip at 0 crossing
+			const int32 WrapIndex = FMath::TruncToInt32(InOutIndex) % InMax;
+			const float Remainder = FMath::Frac(InOutIndex);
+			InOutIndex = WrapIndex + Remainder;
+		};
+
+		void ComputeGainIndexData(const TArray<FWaveTableData>& WaveTables, EWaveTableSamplingMode InSampleMode, ::WaveTable::FWaveTable& OutputWaveTable)
+		{
+			checkf(!WaveTables.IsEmpty(), TEXT("ComputGainIndexData must have WaveTable to operate on"));
+
+			if (WaveTables.Num() == 1)
+			{
+				GainIndexData.LowIndex = INDEX_NONE;
+				GainIndexData.HighIndex = 0;
+				GainIndexData.HighIndexGain = 1.0f;
+				GainIndexData.NumSamples = WaveTables[0].GetNumSamples();
+				return;
+			}
+
+			switch (InSampleMode)
+			{
+				case EWaveTableSamplingMode::FixedSampleRate:
+				{
+					GainIndexData.LowIndex = INDEX_NONE;
+					GainIndexData.HighIndex = FMath::FloorToInt32(GainIndexData.TableIndex);
+					GainIndexData.HighIndexGain = 1.0f;
+					GainIndexData.NumSamples = WaveTables[GainIndexData.HighIndex].GetNumSamples();
+				}
+				break;
+
+				case EWaveTableSamplingMode::FixedResolution:
+				{
+					GainIndexData.LowIndex = FMath::FloorToInt32(GainIndexData.TableIndex);
+					GainIndexData.HighIndex = FMath::CeilToInt32(GainIndexData.TableIndex) % WaveTables.Num();
+
+					// If low/high are same index, give all energy to high index and invalidate
+					// low index (the float index is computed to be an integer equivalent)
+					if (GainIndexData.LowIndex == GainIndexData.HighIndex)
+					{
+						GainIndexData.LowIndex = INDEX_NONE;
+						GainIndexData.HighIndexGain = 1.0f;
+					}
+					else
+					{
+						GainIndexData.HighIndexGain = GainIndexData.TableIndex - GainIndexData.LowIndex;
+					}
+
+					GainIndexData.NumSamples = WaveTables.Last().GetNumSamples();
+				}
+				break;
+
+				default:
+				{
+					static_assert(static_cast<int32>(EWaveTableSamplingMode::COUNT) == 2, "Possible missing switch case coverage for 'EWaveTableSamplingMode'");
+					checkNoEntry();
+				}
+			};
+		}
+
+	public:
 		void Execute()
 		{
 			using namespace WaveTable;
@@ -117,106 +185,103 @@ namespace Metasound
 			METASOUND_TRACE_CPUPROFILER_EVENT_SCOPE(FMetasoundWaveTableGetNodeOperator::Execute);
 
 			const FWaveTableBankAsset& WaveTableBankAsset = *WaveTableBankReadRef;
+			FWaveTable& OutputWaveTable = *OutTable;
 
 			FWaveTableBankAssetProxyPtr Proxy = WaveTableBankAsset.GetProxy();
 			if (Proxy.IsValid() == false)
 			{
-				LastTableIndex = -1.0f;
+				ResetInternal();
 				return;
 			}
 			
-			const TArray<FWaveTable>& WaveTables = Proxy->GetWaveTables();
+			const TArray<FWaveTableData>& WaveTables = Proxy->GetWaveTableData();
 			if (WaveTables.IsEmpty())
 			{
-				LastTableIndex = -1.0f;
+				ResetInternal();
 				return;
 			}
-			
-			int32 NumSamples = WaveTables.Last().Num();
-			FWaveTable& OutputWaveTable = *OutTable;
-			OutputWaveTable.SetNum(NumSamples);
-			OutputWaveTable.Zero();
 
-			if (NumSamples <= 0)
+			GainIndexData.TableIndex = *TableIndexReadRef;
+			WrapIndex(WaveTables.Num(), GainIndexData.TableIndex);
+
+			const bool bIsLastIndex = FMath::IsNearlyEqual(GainIndexData.LastTableIndex, GainIndexData.TableIndex);
+			const bool bIsLastProxy = LastTableId == Proxy->GetObjectId();
+			if (bIsLastIndex && bIsLastProxy)
 			{
-				LastTableIndex = -1.0f;
 				return;
 			}
-			
-			auto WrapIndex = [Max = WaveTables.Num()](float& InOutIndex)
-			{
-				InOutIndex = FMath::Abs(InOutIndex); // Avoids remainder offset flip at 0 crossing
-				const int32 WrapIndex = FMath::TruncToInt32(InOutIndex) % Max;
-				const float Remainder = FMath::Frac(InOutIndex);
-				InOutIndex = WrapIndex + Remainder;
-			};
 
-			float NextTableIndex = *TableIndexReadRef;
-			WrapIndex(NextTableIndex);
+			LastTableId = Proxy->GetObjectId();
+			GainIndexData.LastTableIndex = GainIndexData.TableIndex;
 
-			if (WaveTables.Num() == 1)
+			ComputeGainIndexData(WaveTables, Proxy->GetSampleMode(), OutputWaveTable);
+			if (GainIndexData.NumSamples <= 0)
 			{
-				OutputWaveTable = WaveTables.Last();
+				LastTableId = INDEX_NONE;
+				OutputWaveTable = { };
+				return;
 			}
-			else
+
+			float FinalValue = 0.0f;
+			OutputWaveTable.SetNum(GainIndexData.NumSamples);
+			OutputWaveTable.Zero();
+			for (int32 TableIndex = 0; TableIndex < WaveTables.Num(); ++TableIndex)
 			{
-				if (LastTableIndex < 0.0f)
+				float Gain = 0.0f;
+
+				if (TableIndex == GainIndexData.LowIndex)
 				{
-					LastTableIndex = NextTableIndex;
+					Gain = 1.0f - GainIndexData.HighIndexGain;
 				}
-				else
+				else if (TableIndex == GainIndexData.HighIndex)
 				{
-					WrapIndex(LastTableIndex);
+					Gain = GainIndexData.HighIndexGain;
 				}
 
-				const int32 StartHighIndex = FMath::CeilToInt32(LastTableIndex) % WaveTables.Num();
-				const int32 StartLowIndex = FMath::FloorToInt32(LastTableIndex);
-
-				const float StartHighGain = LastTableIndex - StartLowIndex;
-				const float StartLowGain = 1.0f - StartHighGain;
-
-				float FinalValue = 0.f;
-
-				for (int32 TableIndex = 0; TableIndex < WaveTables.Num(); ++TableIndex)
+				if (Gain > 0.0f)
 				{
-					float Gain = 0.0f;
-
-					if (TableIndex == StartLowIndex)
-					{
-						Gain = StartLowGain;
-					}
-					else if (TableIndex == StartHighIndex)
-					{
-						Gain = StartHighGain;
-					}
-
-					if (Gain > 0.0f)
-					{
-						const FWaveTableView& InputWaveTable = WaveTables[TableIndex].GetView();
-						
-						Audio::ArrayMixIn(InputWaveTable.SampleView, OutputWaveTable.GetSamples(), Gain);
-						FinalValue += InputWaveTable.FinalValue * Gain;
-					}
+					const FWaveTableData& InputTableData = WaveTables[TableIndex];
+					InputTableData.ArrayMixIn(OutputWaveTable.GetSamples(), Gain);
+					FinalValue += InputTableData.GetFinalValue() * Gain;
 				}
-				
-				OutputWaveTable.SetFinalValue(FinalValue);
 			}
 
-			LastTableIndex = NextTableIndex;
+			OutputWaveTable.SetFinalValue(FinalValue);
 		}
 
 		void Reset(const IOperator::FResetParams& InParams)
 		{
-			LastTableIndex = -1.0f;
-			OutTable->SetNum(0);
-			OutTable->SetFinalValue(0.f);
+			BlockPeriod = InParams.OperatorSettings.GetActualBlockRate() / InParams.OperatorSettings.GetSampleRate();
+			ResetInternal();
 		}
 
 	private:
+		void ResetInternal()
+		{
+			GainIndexData = { };
+			LastTableId = INDEX_NONE;
+			*OutTable = { };
+		}
+
 		FWaveTableBankAssetReadRef WaveTableBankReadRef;
 		FFloatReadRef TableIndexReadRef;
 
-		float LastTableIndex = -1.0f;
+		struct FGainIndexData
+		{
+			int32 NumSamples = 0;
+
+			int32 LowIndex = 0;
+			int32 HighIndex = INDEX_NONE;
+
+			float LastTableIndex = -1.0f;
+			float TableIndex = 0.0f;
+
+			float HighIndexGain = 1.0f;
+		} GainIndexData;
+
+		float BlockPeriod = 0.0f;
+
+		uint32 LastTableId = INDEX_NONE;
 
 		FWaveTableWriteRef OutTable;
 	};
diff --git a/Engine/Plugins/Runtime/Metasound/Source/MetasoundEngine/Private/MetasoundWaveTablePlayerNode.cpp b/Engine/Plugins/Runtime/Metasound/Source/MetasoundEngine/Private/MetasoundWaveTablePlayerNode.cpp
new file mode 100644
--- /dev/null
+++ b/Engine/Plugins/Runtime/Metasound/Source/MetasoundEngine/Private/MetasoundWaveTablePlayerNode.cpp
@@ -0,0 +1,353 @@
+// Copyright Epic Games, Inc. All Rights Reserved.
+
+#include "Interfaces/MetasoundFrontendSourceInterface.h"
+#include "MetasoundDataFactory.h"
+#include "MetasoundEngineNodesNames.h"
+#include "MetasoundExecutableOperator.h"
+#include "MetasoundFacade.h"
+#include "MetasoundNodeRegistrationMacro.h"
+#include "MetasoundPrimitives.h"
+#include "MetasoundStandardNodesCategories.h"
+#include "MetasoundStandardNodesNames.h"
+#include "MetasoundTrace.h"
+#include "MetasoundTrigger.h"
+#include "MetasoundWaveTable.h"
+#include "WaveTableSampler.h"
+
+#define LOCTEXT_NAMESPACE "MetasoundStandardNodes"
+
+
+namespace Metasound
+{
+	class FMetaSoundWaveTablePlayerNodeOperator : public TExecutableOperator<FMetaSoundWaveTablePlayerNodeOperator>
+	{
+	public:
+		// Maximum absolute pitch shift in semitones. 
+		static constexpr float MaxAbsPitchShiftInSemitones = 72.0f;
+
+		static const FVertexInterface& GetDefaultInterface()
+		{
+			using namespace WaveTable;
+
+			static const FVertexInterface DefaultInterface(
+				FInputVertexInterface(
+					TInputDataVertex<FTrigger>("Play", FDataVertexMetadata { LOCTEXT("MetaSoundWaveTablePlayerNode_InputPlayDesc", "Plays the wavetable") }),
+					TInputDataVertex<FTrigger>("Stop", FDataVertexMetadata { LOCTEXT("MetaSoundWaveTablePlayerNode_InputStopDesc", "Stops the wavetable") }),
+					TInputDataVertex<FTrigger>("Sync", FDataVertexMetadata
+					{
+						LOCTEXT("MetaSoundWaveTablePlayerNode_InputSyncDesc", "Restarts playing the WaveTable on the trigger boundary (sample rate)"),
+						LOCTEXT("MetaSoundWaveTablePlayerNode_InputSyncName", "Sync")
+					}),
+					TInputDataVertex<FWaveTableBankAsset>("WaveTableBank", FDataVertexMetadata
+					{
+						LOCTEXT("MetaSoundWaveTablePlayerNode_InputBankDesc", "WaveTable Bank to playback from"),
+						LOCTEXT("MetaSoundWaveTablePlayerNode_InputBankName", "Bank")
+					}),
+					TInputDataVertex<int32>("Index", FDataVertexMetadata
+					{
+						LOCTEXT("MetaSoundWaveTablePlayerNode_InputIndexDesc", "Index of WaveTable Bank entry to play"),
+						LOCTEXT("MetaSoundWaveTablePlayerNode_InputIndexName", "Index")
+					}),
+					TInputDataVertex<float>("PitchShift", FDataVertexMetadata
+					{
+						LOCTEXT("MetaSoundWaveTablePlayerNode_PitchShiftDesc", "How much to shift the pitch of the given WaveTable"),
+						LOCTEXT("MetaSoundWaveTablePlayerNode_PitchShiftName", "Pitch Shift")
+					}, 0.0f),
+					TInputDataVertex<bool>("Loop", FDataVertexMetadata { LOCTEXT("MetaSoundWaveTablePlayerNode_LoopDesc", "Whether or not to loop the given WaveTable") }, false),
+					TInputDataVertex<FAudioBuffer>("PhaseMod", FDataVertexMetadata
+					{
+						LOCTEXT("MetaSoundWaveTablePlayerNode_PhaseModDescription", "Modulation audio source for modulating oscillation phase of provided table. A value of 0 is no phase modulation and 1 a full table length (360 degrees) of phase shift."),
+						LOCTEXT("MetaSoundWaveTablePlayerNode_PhaseMod", "Phase Modulator"),
+						true /* bIsAdvancedDisplay */
+					})
+				),
+				FOutputVertexInterface(
+					TOutputDataVertex<FAudioBuffer>("MonoOut", FDataVertexMetadata
+					{
+						LOCTEXT("MetaSoundWaveTablePlayer_OutputMonoOutDesc", "Mono output audio from player"),
+						LOCTEXT("MetaSoundWaveTablePlayer_OutputMonoOutName", "Mono Out")
+					}),
+					TOutputDataVertex<FTrigger>("OnFinished", FDataVertexMetadata
+					{
+						LOCTEXT("MetaSoundWaveTablePlayerNode_OnFinished", "Trigger executed when player is complete (if not looping) or stop is triggered"),
+						LOCTEXT("MetaSoundWaveTablePlayer_OutputOnFinishedName", "On Finished")
+					})
+				)
+			);
+
+			return DefaultInterface;
+		}
+
+		static const FNodeClassMetadata& GetNodeInfo()
+		{
+			auto CreateNodeClassMetadata = []() -> FNodeClassMetadata
+			{
+				FNodeClassMetadata Metadata
+				{
+					{ EngineNodes::Namespace, "WaveTablePlayer", "" },
+					1, // Major Version
+					0, // Minor Version
+					LOCTEXT("MetaSoundWaveTablePlayerNode_Name", "WaveTable Player"),
+					LOCTEXT("MetaSoundWaveTablePlayerNode_Description", "Reads through the given WaveTableBank's entry at the given index."),
+					PluginAuthor,
+					PluginNodeMissingPrompt,
+					GetDefaultInterface(),
+					{ NodeCategories::Generators },
+					{ METASOUND_LOCTEXT("WaveTablePlayerSynthesisKeyword", "Synthesis")},
+					{ }
+				};
+
+				return Metadata;
+			};
+
+			static const FNodeClassMetadata Metadata = CreateNodeClassMetadata();
+			return Metadata;
+		}
+
+		static TUniquePtr<IOperator> CreateOperator(const FCreateOperatorParams& InParams, TArray<TUniquePtr<IOperatorBuildError>>& OutErrors)
+		{
+			using namespace WaveTable;
+
+			const FInputVertexInterface& InputInterface = InParams.Node.GetVertexInterface().GetInputInterface();
+			const FDataReferenceCollection& InputCollection = InParams.InputDataReferences;
+
+			FWaveTableBankAssetReadRef InWaveTableBankReadRef = InputCollection.GetDataReadReferenceOrConstruct<FWaveTableBankAsset>("WaveTableBank");
+			FInt32ReadRef InIndexReadRef = InputCollection.GetDataReadReferenceOrConstruct<int32>("Index");
+			FTriggerReadRef InPlayReadRef = InputCollection.GetDataReadReferenceOrConstructWithVertexDefault<FTrigger>(InputInterface, "Play", InParams.OperatorSettings);
+			FTriggerReadRef InStopReadRef = InputCollection.GetDataReadReferenceOrConstructWithVertexDefault<FTrigger>(InputInterface, "Stop", InParams.OperatorSettings);
+			FTriggerReadRef InSyncReadRef = InputCollection.GetDataReadReferenceOrConstructWithVertexDefault<FTrigger>(InputInterface, "Sync", InParams.OperatorSettings);
+			FFloatReadRef InPitchShiftReadRef = InputCollection.GetDataReadReferenceOrConstructWithVertexDefault<float>(InputInterface, "PitchShift", InParams.OperatorSettings);
+			FBoolReadRef InLoopReadRef = InputCollection.GetDataReadReferenceOrConstructWithVertexDefault<bool>(InputInterface, "Loop", InParams.OperatorSettings);
+
+			TOptional<FAudioBufferReadRef> InPhaseModReadRef;
+			if (InputCollection.ContainsDataReadReference<FAudioBuffer>("PhaseMod"))
+			{
+				InPhaseModReadRef = InputCollection.GetDataReadReference<FAudioBuffer>("PhaseMod");
+			}
+
+			return MakeUnique<FMetaSoundWaveTablePlayerNodeOperator>(InParams, InWaveTableBankReadRef, InIndexReadRef, InPlayReadRef, InStopReadRef, InSyncReadRef, InPitchShiftReadRef, InLoopReadRef, MoveTemp(InPhaseModReadRef));
+		}
+
+		FMetaSoundWaveTablePlayerNodeOperator(
+			const FCreateOperatorParams& InParams,
+			const FWaveTableBankAssetReadRef& InWaveTableBankReadRef,
+			const FInt32ReadRef& InIndexReadRef,
+			const FTriggerReadRef& InPlayReadRef,
+			const FTriggerReadRef& InStopReadRef,
+			const FTriggerReadRef& InSyncReadRef,
+			const FFloatReadRef& InPitchShiftReadRef,
+			const FBoolReadRef& InLoopReadRef,
+			TOptional<FAudioBufferReadRef>&& InPhaseModReadRef
+		)
+			: SampleRate(InParams.OperatorSettings.GetSampleRate())
+			, WaveTableBankReadRef(InWaveTableBankReadRef)
+			, IndexReadRef(InIndexReadRef)
+			, PlayReadRef(InPlayReadRef)
+			, StopReadRef(InStopReadRef)
+			, SyncReadRef(InSyncReadRef)
+			, PitchShiftReadRef(InPitchShiftReadRef)
+			, LoopReadRef(InLoopReadRef)
+			, PhaseModReadRef(MoveTemp(InPhaseModReadRef))
+			, OutBufferWriteRef(TDataWriteReferenceFactory<FAudioBuffer>::CreateAny(InParams.OperatorSettings))
+			, OutOnFinishedRef(TDataWriteReferenceFactory<FTrigger>::CreateExplicitArgs(InParams.OperatorSettings))
+		{
+			const float BlockRate = InParams.OperatorSettings.GetActualBlockRate();
+			if (BlockRate > 0.0f)
+			{
+				BlockPeriod = 1.0f / BlockRate;
+			}
+		}
+
+		virtual ~FMetaSoundWaveTablePlayerNodeOperator() = default;
+
+		virtual FDataReferenceCollection GetInputs() const override
+		{
+			FDataReferenceCollection Inputs;
+
+			Inputs.AddDataReadReference("WaveTableBank", WaveTableBankReadRef);
+			Inputs.AddDataReadReference("Index", IndexReadRef);
+			Inputs.AddDataReadReference("Sync", SyncReadRef);
+			Inputs.AddDataReadReference("PitchShift", PitchShiftReadRef);
+			Inputs.AddDataReadReference("Loop", LoopReadRef);
+
+			if (PhaseModReadRef.IsSet())
+			{
+				Inputs.AddDataReadReference("PhaseMod", *PhaseModReadRef);
+			}
+
+			return Inputs;
+		}
+
+		virtual FDataReferenceCollection GetOutputs() const override
+		{
+			FDataReferenceCollection Outputs;
+			Outputs.AddDataReadReference("MonoOut", TDataReadReference<FAudioBuffer>(OutBufferWriteRef));
+			Outputs.AddDataReadReference("OnFinished", TDataReadReference<FTrigger>(OutOnFinishedRef));
+
+			return Outputs;
+		}
+
+		void Execute()
+		{
+			using namespace WaveTable;
+
+			METASOUND_TRACE_CPUPROFILER_EVENT_SCOPE(FMetaSoundWaveTablePlayerNodeOperator::Execute);
+
+			FTrigger& OnFinishedTrigger = *OutOnFinishedRef;
+			OnFinishedTrigger.AdvanceBlock();
+
+			FAudioBuffer& OutBuffer = *OutBufferWriteRef;
+			OutBuffer.Zero();
+			auto GetLastTriggerFrame = [](const FTriggerReadRef& Trigger)
+			{
+				int32 LastIndex = INDEX_NONE;
+				Trigger->ExecuteBlock([](int32, int32) {}, [&LastIndex](int32 StartFrame, int32 EndFrame)
+				{
+					LastIndex = FMath::Max(LastIndex, StartFrame);
+				});
+				return LastIndex;
+			};
+
+			bool bIsStopping = false;
+			const int32 LastPlayFrame = GetLastTriggerFrame(PlayReadRef);
+			int32 LastCompleteFrame = GetLastTriggerFrame(StopReadRef);
+			if (LastPlayFrame >= 0 || LastCompleteFrame >= 0)
+			{
+				const bool bWasPlaying = bPlaying;
+				bPlaying = LastPlayFrame > LastCompleteFrame;
+				bIsStopping = bWasPlaying && !bPlaying;
+			}
+
+			if (bPlaying || bIsStopping)
+			{
+				const FTrigger& SyncTrigger = *SyncReadRef;
+				TArrayView<float> SyncBufferView;
+				if (SyncTrigger.IsTriggered())
+				{
+					SyncBuffer.SetNum(OutBuffer.Num());
+					FMemory::Memset(SyncBuffer.GetData(), 0, sizeof(float) * SyncBuffer.Num());
+					SyncTrigger.ExecuteBlock(
+						[](int32 StartFrame, int32 EndFrame) {},
+						[this](int32 StartFrame, int32 EndFrame)
+						{
+							SyncBuffer[StartFrame] = 1.0f;
+						}
+					);
+					SyncBufferView = SyncBuffer;
+				}
+
+				TArrayView<const float> PhaseMod;
+				if (PhaseModReadRef.IsSet())
+				{
+					const FAudioBuffer& Buffer = *(*PhaseModReadRef);
+					PhaseMod = { Buffer.GetData(), Buffer.Num() };
+				}
+
+				const FWaveTableBankAsset& BankAsset = *WaveTableBankReadRef;
+				FWaveTableBankAssetProxyPtr BankProxy = BankAsset.GetProxy();
+				if (BankProxy.IsValid())
+				{
+					const float BankSampleRate = BankProxy->GetSampleMode() == EWaveTableSamplingMode::FixedResolution ? SampleRate : (float)BankProxy->GetSampleRate();
+					float FreqRatio = 1.0f;
+					if (ensureAlwaysMsgf(SampleRate > 0, TEXT("SampleRate must always be non-zero and positive")))
+					{
+						const float PitchShift = FMath::Clamp(*PitchShiftReadRef, -MaxAbsPitchShiftInSemitones, MaxAbsPitchShiftInSemitones);
+						FreqRatio = Audio::GetFrequencyMultiplier(PitchShift) * (BankSampleRate / SampleRate);
+					}
+
+					const TArray<FWaveTableData>& WaveTableData = BankProxy->GetWaveTableData();
+					if (!WaveTableData.IsEmpty())
+					{
+						const int32 TableIndex = FMath::Abs(*IndexReadRef) % WaveTableData.Num();
+
+						const FWaveTableData& Entry = WaveTableData[TableIndex];
+						const float Duration = Entry.GetNumSamples() / SampleRate;
+						const float Freq = FreqRatio / Duration;
+						Sampler.SetFreq(Freq * BlockPeriod);
+
+						const bool bIsLooping = *LoopReadRef;
+						Sampler.SetOneShot(!bIsLooping);
+
+						Sampler.Process(WaveTableData[TableIndex], {}, PhaseMod, SyncBufferView, OutBuffer);
+
+						if (LastCompleteFrame > 0)
+						{
+							checkf(LastCompleteFrame < OutBuffer.Num(), TEXT("LastCompleteFrame/output buffer size mismatch"));
+							FMemory::Memzero(OutBuffer.GetData() + LastCompleteFrame, sizeof(float) * (OutBuffer.Num() - LastCompleteFrame));
+						}
+						else
+						{
+							const int32 IndexFinished = Sampler.GetIndexFinished();
+							if (IndexFinished >= 0)
+							{
+								bIsStopping = true;
+								bPlaying = false;
+								LastCompleteFrame = IndexFinished;
+							}
+						}
+					}
+				}
+			}
+
+			if (bIsStopping)
+			{
+				OnFinishedTrigger.TriggerFrame(LastCompleteFrame);
+			}
+		}
+
+		void Reset(const IOperator::FResetParams& InParams)
+		{
+			const float BlockRate = InParams.OperatorSettings.GetActualBlockRate();
+			SampleRate = InParams.OperatorSettings.GetSampleRate();
+			if (BlockRate > 0.0f)
+			{
+				BlockPeriod = 1.0f / BlockRate;
+			}
+
+			bPlaying = false;
+			if (SyncBuffer.Num() > 0)
+			{
+				FMemory::Memset(SyncBuffer.GetData(), 0, sizeof(float) * SyncBuffer.Num());
+			}
+			Sampler = WaveTable::FWaveTableSampler{};
+			OutBufferWriteRef->Zero();
+			OutOnFinishedRef->Reset();
+		}
+
+	private:
+		float BlockPeriod = 0.0f;
+		float SampleRate = 48000.0f;
+		bool bPlaying = false;
+
+		FWaveTableBankAssetReadRef WaveTableBankReadRef;
+		FInt32ReadRef IndexReadRef;
+		FTriggerReadRef PlayReadRef;
+		FTriggerReadRef StopReadRef;
+		FTriggerReadRef SyncReadRef;
+		FFloatReadRef PitchShiftReadRef;
+		FBoolReadRef LoopReadRef;
+
+		Audio::FAlignedFloatBuffer SyncBuffer;
+		TOptional<FAudioBufferReadRef> PhaseModReadRef;
+
+		WaveTable::FWaveTableSampler Sampler;
+
+		TDataWriteReference<FAudioBuffer> OutBufferWriteRef;
+		TDataWriteReference<FTrigger> OutOnFinishedRef;
+	};
+
+	class FMetaSoundWaveTablePlayerNode : public FNodeFacade
+	{
+	public:
+		FMetaSoundWaveTablePlayerNode(const FNodeInitData& InInitData)
+			: FNodeFacade(InInitData.InstanceName, InInitData.InstanceID, TFacadeOperatorClass<FMetaSoundWaveTablePlayerNodeOperator>())
+		{
+		}
+
+		virtual ~FMetaSoundWaveTablePlayerNode() = default;
+	};
+
+	METASOUND_REGISTER_NODE(FMetaSoundWaveTablePlayerNode)
+} // namespace Metasound
+
+#undef LOCTEXT_NAMESPACE // MetasoundStandardNodes
diff --git a/Engine/Plugins/Runtime/WaveTable/Source/WaveTable/Private/WaveTable.cpp b/Engine/Plugins/Runtime/WaveTable/Source/WaveTable/Private/WaveTable.cpp
new file mode 100644
--- /dev/null
+++ b/Engine/Plugins/Runtime/WaveTable/Source/WaveTable/Private/WaveTable.cpp
@@ -0,0 +1,476 @@
+// Copyright Epic Games, Inc. All Rights Reserved.
+#include "WaveTable.h"
+
+#include "DSP/FloatArrayMath.h"
+#include "WaveTableSampler.h"
+
+#include UE_INLINE_GENERATED_CPP_BY_NAME(WaveTable)
+
+
+namespace WaveTable
+{
+	FWaveTableView::FWaveTableView(const TArray<float>& InSamples, const float InFinalValue)
+		: SampleView(InSamples)
+		, FinalValue(InFinalValue)
+	{
+	}
+
+	FWaveTableView::FWaveTableView(const TArrayView<const float>& InSamples, const float InFinalValue)
+		: SampleView(InSamples)
+		, FinalValue(InFinalValue)
+	{
+	}
+
+	int32 FWaveTableView::Num() const
+	{
+		return SampleView.Num();
+	}
+
+	bool FWaveTableView::IsEmpty() const
+	{
+		return SampleView.IsEmpty();
+	}
+
+	FWaveTable::FWaveTable(const TArray<float>& InSamples, const float InFinalValue)
+		: Samples(InSamples)
+		, FinalValue(InFinalValue)
+	{
+	}
+
+	FWaveTable::FWaveTable(TArray<float>&& InSamples, const float InFinalValue)
+		: Samples(MoveTemp(InSamples))
+		, FinalValue(InFinalValue)
+	{
+	}
+
+	FWaveTable::FWaveTable(const ::FWaveTableData& InData)
+		: FinalValue(InData.GetFinalValue())
+	{
+		switch (InData.GetBitDepth())
+		{
+			case EWaveTableBitDepth::PCM_16:
+			{
+				TArrayView<const int16> Samples16Bit;
+				ensureAlwaysMsgf(InData.GetDataView(Samples16Bit), TEXT("PCM_16 BitDepth format mismatch"));
+				Samples.AddUninitialized(Samples16Bit.Num());
+				Audio::ArrayPcm16ToFloat(Samples16Bit, Samples);
+			}
+			break;
+
+			case EWaveTableBitDepth::IEEE_Float:
+			{
+				TArrayView<const float> SamplesFloat;
+				ensureAlwaysMsgf(InData.GetDataView(SamplesFloat), TEXT("IEEE_Float BitDepth format mismatch"));
+				Samples = SamplesFloat;
+			}
+			break;
+
+			default:
+			{
+				static_assert(static_cast<int32>(EWaveTableBitDepth::COUNT) == 2, "Possible missing switch case coverage for 'EWaveTableBitDepth'");
+				checkNoEntry();
+			}
+			break;
+		}
+	}
+
+	float* FWaveTable::GetData()
+	{
+		return Samples.GetData();
+	}
+
+	const float* FWaveTable::GetData() const
+	{
+		return Samples.GetData();
+	}
+
+	FWaveTableView FWaveTable::GetView() const
+	{
+		return FWaveTableView(Samples, FinalValue);
+	}
+
+	TArrayView<float> FWaveTable::GetSamples()
+	{
+		return Samples;
+	}
+
+	int32 FWaveTable::Num() const
+	{
+		return Samples.Num();
+	}
+
+	void FWaveTable::Set(TArray<float>&& InSamples)
+	{
+		Samples = MoveTemp(InSamples);
+	}
+
+	void FWaveTable::Set(TArrayView<const float> InSamples)
+	{
+		Samples = InSamples;
+	}
+
+	void FWaveTable::SetFinalValue(const float InValue)
+	{
+		FinalValue = InValue;
+	}
+
+	void FWaveTable::SetNum(int32 InSize)
+	{
+		Samples.SetNum(InSize);
+	}
+
+	void FWaveTable::Zero()
+	{
+		FMemory::Memset(Samples.GetData(), 0, sizeof(float) * Samples.Num());
+	}
+} // namespace WaveTable
+
+
+FWaveTableData::FWaveTableData(EWaveTableBitDepth InBitDepth)
+	: BitDepth(InBitDepth)
+{
+}
+
+FWaveTableData::FWaveTableData(TArrayView<const float> InSamples, float InFinalValue)
+	: BitDepth(EWaveTableBitDepth::IEEE_Float)
+	, Data(TArrayView<uint8>((uint8*)InSamples.GetData(), InSamples.Num() * sizeof(float)))
+	, FinalValue(InFinalValue)
+{
+}
+
+FWaveTableData::FWaveTableData(TArrayView<const int16> InSamples, float InFinalValue)
+	: BitDepth(EWaveTableBitDepth::PCM_16)
+	, Data(TArrayView<uint8>((uint8*)InSamples.GetData(), InSamples.Num() * sizeof(int16)))
+	, FinalValue(InFinalValue)
+{
+}
+
+#if WITH_EDITOR
+FName FWaveTableData::GetBitDepthPropertyName()
+{
+	return GET_MEMBER_NAME_CHECKED(FWaveTableData, BitDepth);
+}
+#endif // WITH_EDITOR
+
+void FWaveTableData::ArrayMixIn(TArrayView<float> OutFloatBuffer, float Gain) const
+{
+	checkf(OutFloatBuffer.Num() == GetNumSamples(), TEXT("Float buffer size must match that of WaveTableData"));
+	switch (BitDepth)
+	{
+		case EWaveTableBitDepth::IEEE_Float:
+		{
+			TArrayView<const float> DataView;
+			GetDataView(DataView);
+			Audio::ArrayMixIn(DataView, OutFloatBuffer, Gain);
+		}
+		break;
+
+		case EWaveTableBitDepth::PCM_16:
+		{
+			TArrayView<const int16> DataView;
+			GetDataView(DataView);
+			Audio::ArrayMixIn(DataView, OutFloatBuffer, Gain);
+		}
+		break;
+
+		default:
+		{
+			static_assert(static_cast<int32>(EWaveTableBitDepth::COUNT) == 2, "Possible missing switch case coverage for 'EWaveTableBitDepth'");
+			checkNoEntry();
+		}
+		break;
+	}
+}
+
+int32 FWaveTableData::GetSampleSize() const
+{
+	switch (BitDepth)
+	{
+
+		case EWaveTableBitDepth::PCM_16:
+		{
+			return sizeof(int16);
+		}
+
+		case EWaveTableBitDepth::IEEE_Float:
+		{
+			return sizeof(float);
+		}
+
+		default:
+		{
+			static_assert(static_cast<int32>(EWaveTableBitDepth::COUNT) == 2, "Possible missing switch case coverage for 'EWaveTableBitDepth'");
+			checkNoEntry();
+			return sizeof(uint8);
+		}
+		break;
+	}
+}
+
+void FWaveTableData::Empty()
+{
+	Data.Empty();
+	FinalValue = ::WaveTable::InvalidWaveTableValue;
+}
+
+bool FWaveTableData::Resample(int32 InCurrentSampleRate, int32 InNewSampleRate, ::WaveTable::FWaveTableSampler::EInterpolationMode InInterpMode)
+{
+	if (InCurrentSampleRate != InNewSampleRate)
+	{
+
+		const EWaveTableBitDepth LastBitDepth = BitDepth;
+		FWaveTableData ResampledData(EWaveTableBitDepth::IEEE_Float);
+
+		// 1. Init scratch WaveTableData for resampled data (must be in IEEE_Float)
+		{
+			ResampledData.Zero(GetNumSamples() * InCurrentSampleRate / InNewSampleRate);
+		}
+
+		// 2. Resample in correct bit depth
+		{
+			TArrayView<float> ResampledView;
+			ResampledData.GetDataView(ResampledView);
+			::WaveTable::FWaveTableSampler::FSettings Settings;
+			Settings.InterpolationMode = InInterpMode;
+			::WaveTable::FWaveTableSampler Sampler(MoveTemp(Settings));
+			Sampler.Process(*this, ResampledView);
+		}
+
+		// 3. Move data & set to correct bit depth (ignored if original bit depth was already IEEE_Float)
+		{
+			*this = MoveTemp(ResampledData);
+			SetBitDepth(LastBitDepth);
+		}
+
+		return true;
+	}
+
+	return false;
+}
+
+void FWaveTableData::Reset(int32 NumSamples)
+{
+	Data.Reset(NumSamples * GetSampleSize());
+	FinalValue = ::WaveTable::InvalidWaveTableValue;
+}
+
+bool FWaveTableData::SetBitDepth(EWaveTableBitDepth InNewBitDepth)
+{
+	if (InNewBitDepth == BitDepth)
+	{
+		return false;
+	}
+
+	const int32 NumSamples = GetNumSamples();
+	TArray<uint8> ConvertedData;
+	switch(InNewBitDepth)
+	{
+		case EWaveTableBitDepth::IEEE_Float:
+		{
+			ConvertedData.SetNumUninitialized(sizeof(float) * NumSamples);
+
+			TArrayView<const int16> InputWaveTable;
+			GetDataView(InputWaveTable);
+			Audio::ArrayPcm16ToFloat(InputWaveTable, TArrayView<float>((float*)ConvertedData.GetData(), ConvertedData.Num()));
+		}
+		break;
+
+		case EWaveTableBitDepth::PCM_16:
+		{
+			ConvertedData.SetNumUninitialized(sizeof(int16) * NumSamples);
+
+			TArrayView<const float> InputWaveTable;
+			GetDataView(InputWaveTable);
+			Audio::ArrayFloatToPcm16(InputWaveTable, TArrayView<int16>((int16*)ConvertedData.GetData(), ConvertedData.Num()));
+		}
+		break;
+
+		default:
+		{
+			static_assert(static_cast<int32>(EWaveTableBitDepth::COUNT) == 2, "Possible missing switch case coverage for 'EWaveTableBitDepth'");
+			checkNoEntry();
+		}
+	}
+
+	BitDepth = InNewBitDepth;
+	Data = MoveTemp(ConvertedData);
+	return true;
+}
+
+void FWaveTableData::SetData(TArrayView<const int16> InData, bool bIsLoop)
+{
+	Data = TArrayView<uint8>((uint8*)InData.GetData(), InData.Num() * sizeof(int16));
+	BitDepth = EWaveTableBitDepth::PCM_16;
+	if (InData.IsEmpty())
+	{
+		FinalValue = 0.0f;
+	}
+	else
+	{
+		constexpr float ConversionValue = 1.0f / static_cast<float>(TNumericLimits<int16>::Max());
+		if (bIsLoop)
+		{
+			FinalValue = InData[0] * ConversionValue;
+		}
+		else
+		{
+			FinalValue = InData.Last() * ConversionValue;
+		}
+	}
+}
+
+void FWaveTableData::SetData(TArrayView<const float> InData, bool bIsLoop)
+{
+	Data = TArrayView<uint8>((uint8*)InData.GetData(), InData.Num() * sizeof(float));
+	BitDepth = EWaveTableBitDepth::IEEE_Float;
+	if (InData.IsEmpty())
+	{
+		FinalValue = 0.0f;
+	}
+	else
+	{
+		if (bIsLoop)
+		{
+			FinalValue = InData[0];
+		}
+		else
+		{
+			FinalValue = InData.Last();
+		}
+	}
+}
+
+void FWaveTableData::SetFinalValue(float InFinalValue)
+{
+	FinalValue = InFinalValue;
+}
+
+void FWaveTableData::SetRawData(EWaveTableBitDepth InBitDepth, TArrayView<const uint8> InData)
+{
+	Data = InData;
+	BitDepth = InBitDepth;
+}
+
+void FWaveTableData::SetRawData(EWaveTableBitDepth InBitDepth, TArray<uint8>&& InData)
+{
+	BitDepth = InBitDepth;
+	Data = MoveTemp(InData);
+}
+
+EWaveTableBitDepth FWaveTableData::GetBitDepth() const
+{
+	return BitDepth;
+}
+
+bool FWaveTableData::GetDataView(TArrayView<const float>& OutData) const
+{
+	if (BitDepth == EWaveTableBitDepth::IEEE_Float)
+	{
+		OutData = TArrayView<const float>((const float*)(Data.GetData()), Data.Num() / sizeof(float));
+		return true;
+	}
+
+	return false;
+}
+
+bool FWaveTableData::GetDataView(TArrayView<float>& OutData)
+{
+	if (BitDepth == EWaveTableBitDepth::IEEE_Float)
+	{
+		OutData = TArrayView<float>((float*)(Data.GetData()), Data.Num() / sizeof(float));
+		return true;
+	}
+
+	return false;
+}
+
+bool FWaveTableData::GetDataView(TArrayView<const int16>& OutData) const
+{
+	if (BitDepth == EWaveTableBitDepth::PCM_16)
+	{
+		OutData = TArrayView<const int16>((const int16*)(Data.GetData()), Data.Num() / sizeof(int16));
+		return true;
+	}
+
+	return false;
+}
+
+bool FWaveTableData::GetDataView(TArrayView<int16>& OutData)
+{
+	if (BitDepth == EWaveTableBitDepth::PCM_16)
+	{
+		OutData = TArrayView<int16>((int16*)(Data.GetData()), Data.Num() / sizeof(int16));
+		return true;
+	}
+
+	return false;
+}
+
+float FWaveTableData::GetFinalValue() const
+{
+	if (FMath::IsNearlyEqual(FinalValue, WaveTable::InvalidWaveTableValue))
+	{
+		if (Data.IsEmpty())
+		{
+			return 0.0f;
+		}
+
+		switch (BitDepth)
+		{
+			case EWaveTableBitDepth::IEEE_Float:
+			{
+				TArrayView<const float> TableView;
+				ensureAlways(GetDataView(TableView));
+				return TableView.Last();
+			}
+			break;
+
+			case EWaveTableBitDepth::PCM_16:
+			{
+				TArrayView<const int16> TableView;
+				ensureAlways(GetDataView(TableView));
+				constexpr float ConversionValue = 1.0f / static_cast<float>(TNumericLimits<int16>::Max());
+				return TableView.Last() * ConversionValue;
+			}
+			break;
+
+			default:
+			{
+				static_assert(static_cast<int32>(EWaveTableBitDepth::COUNT) == 2, "Possible missing switch case coverage for 'EWaveTableBitDepth'");
+				checkNoEntry();
+			}
+			break;
+		}
+	}
+
+	return FinalValue;
+}
+
+int32 FWaveTableData::GetNumSamples() const
+{
+	return Data.Num() / GetSampleSize();
+}
+
+const TArray<uint8>& FWaveTableData::GetRawData() const
+{
+	return Data;
+}
+
+bool FWaveTableData::IsEmpty() const
+{
+	return Data.IsEmpty();
+}
+
+void FWaveTableData::Zero(int32 InNumSamples)
+{
+	if (InNumSamples < 0)
+	{
+		Data.Shrink();
+		FMemory::Memzero(Data.GetData(), sizeof(uint8) * Data.Num());
+	}
+	else
+	{
+		const int32 NumBytes = InNumSamples * GetSampleSize();
+		Data.Init(0, NumBytes);
+	}
+}
diff --git a/Engine/Plugins/Runtime/WaveTable/Source/WaveTable/Private/WaveTableBank.cpp b/Engine/Plugins/Runtime/WaveTable/Source/WaveTable/Private/WaveTableBank.cpp
--- a/Engine/Plugins/Runtime/WaveTable/Source/WaveTable/Private/WaveTableBank.cpp
+++ b/Engine/Plugins/Runtime/WaveTable/Source/WaveTable/Private/WaveTableBank.cpp
@@ -1,54 +1,192 @@
 // Copyright Epic Games, Inc. All Rights Reserved.
 
 #include "WaveTableBank.h"
+
+#include "Algo/AnyOf.h"
+#include "Algo/Transform.h"
 #include "UObject/ObjectSaveContext.h"
 
 #include UE_INLINE_GENERATED_CPP_BY_NAME(WaveTableBank)
 
-#define LOCTEXT_NAMESPACE "WaveTable"
-
 
-TSharedPtr<Audio::IProxyData> UWaveTableBank::CreateProxyData(const Audio::FProxyDataInitParams& InitParams)
+namespace WaveTable
 {
-	return MakeShared<FWaveTableBankAssetProxy>(*this);
-}
-
+	namespace BankPrivate
+	{
 #if WITH_EDITOR
-void UWaveTableBank::RefreshWaveTables()
+		int32 GetEntryMaxNumSamples(const UWaveTableBank& InBank)
+		{
+			int32 MaxSamples = 0;
+
+			PRAGMA_DISABLE_DEPRECATION_WARNINGS
+			for (const FWaveTableBankEntry& Entry : InBank.Entries)
+			{
+				int32 NumSamples = 0;
+				if (Entry.Transform.Curve == EWaveTableCurve::File)
+				{
+					int32 Offset = 0;
+					Entry.Transform.WaveTableSettings.GetEditSourceBounds(Offset, NumSamples);
+				}
+				else
+				{
+					NumSamples = 1 << static_cast<int32>(EWaveTableResolution::Res_Max);
+				}
+				MaxSamples = FMath::Max(MaxSamples, NumSamples);
+			}
+			PRAGMA_ENABLE_DEPRECATION_WARNINGS
+
+			return MaxSamples;
+		}
+
+		int32 GetWaveTableNumSamples(const UWaveTableBank& InBank, const FWaveTableBankEntry& InEntry)
+		{
+			const FWaveTableTransform& Transform = InEntry.Transform;
+			switch (InBank.SampleMode)
+			{
+				case EWaveTableSamplingMode::FixedSampleRate:
+				{
+					if (Transform.Curve == EWaveTableCurve::File)
+					{
+						int32 NumSourceSamples = 0;
+						int32 TopOffset = 0;
+						Transform.WaveTableSettings.GetEditSourceBounds(TopOffset, NumSourceSamples);
+						const float ResampleFactor = InBank.SampleRate / static_cast<float>(Transform.WaveTableSettings.SourceSampleRate);
+						return (int32)(NumSourceSamples * ResampleFactor);
+					}
+					else
+					{
+						return (int32)(InBank.SampleRate * Transform.GetDuration());
+					}
+				}
+				break;
+
+				case EWaveTableSamplingMode::FixedResolution:
+				{
+					const int32 MaxNumSamples = GetEntryMaxNumSamples(InBank);
+					switch (InBank.Resolution)
+					{
+						case EWaveTableResolution::Maximum:
+						{
+							return FMath::Max(MaxNumSamples, 1 << static_cast<int32>(EWaveTableResolution::Res_Max));
+						}
+						break;
+
+						case EWaveTableResolution::None:
+						{
+							EWaveTableResolution CurveRes = EWaveTableResolution::None;
+							switch (Transform.Curve)
+							{
+								case EWaveTableCurve::File:
+								{
+									CurveRes = EWaveTableResolution::Maximum;
+								}
+								break;
+
+								case EWaveTableCurve::Linear:
+								case EWaveTableCurve::Linear_Inv:
+								{
+									CurveRes = EWaveTableResolution::Res_8;
+								}
+								break;
+
+								case EWaveTableCurve::Exp:
+								case EWaveTableCurve::Exp_Inverse:
+								case EWaveTableCurve::Log:
+								{
+									CurveRes = EWaveTableResolution::Res_256;
+								}
+								break;
+
+								case EWaveTableCurve::Sin:
+								case EWaveTableCurve::SCurve:
+								case EWaveTableCurve::Sin_Full:
+								{
+									CurveRes = EWaveTableResolution::Res_64;
+								}
+								break;
+
+								default:
+								{
+									CurveRes = EWaveTableResolution::Res_128;
+								}
+								break;
+							};
+							return 1 << static_cast<int32>(CurveRes);
+						}
+						break;
+
+						default:
+						{
+							return 1 << static_cast<int32>(InBank.Resolution);
+						}
+					}
+
+					return MaxNumSamples;
+				}
+				break;
+
+				default:
+				{
+					checkNoEntry();
+					static_assert(static_cast<int32>(EWaveTableSamplingMode::COUNT) == 2, "Possible missing switch case coverage for 'EWaveTableBitDepth'");
+					return 0;
+				}
+				break;
+			}
+		}
+#endif // WITH_EDITOR
+	} // namespace BankPrivate
+} // namespace WaveTable
+
+TSharedPtr<Audio::IProxyData> UWaveTableBank::CreateProxyData(const Audio::FProxyDataInitParams& InitParams)
 {
-	int32 MaxPCMSize = 0;
-	for (FWaveTableBankEntry& Entry : Entries)
+	if (!ProxyData.IsValid())
 	{
-		TArrayView<const float> EditSourceView = Entry.Transform.WaveTableSettings.GetEditSourceView();
-		MaxPCMSize = FMath::Max(MaxPCMSize, EditSourceView.Num());
+#if WITH_EDITOR
+		CopyToProxyData();
+#else
+		MoveToProxyData();
+#endif // WITH_EDITOR
 	}
 
-	WaveTableSizeMB = 0.0f;
-	WaveTableLengthSec = 0.0f;
-	for (FWaveTableBankEntry& Entry : Entries)
-	{
-		TArray<float>& TransformWaveTable = Entry.Transform.WaveTable;
-		TransformWaveTable.Empty();
+	check(ProxyData.IsValid());
+	return ProxyData;
+}
 
-		const int32 WaveTableSize = WaveTable::GetWaveTableSize(Resolution, Entry.Transform.Curve, MaxPCMSize);
-		TransformWaveTable.AddZeroed(WaveTableSize);
+void UWaveTableBank::CopyToProxyData()
+{
+	checkf(!ProxyData.IsValid(), TEXT("Cannot overwrite pre-existing proxy when calling UWaveTableBank::CopyToProxyData"));
 
-		float FinalValue = Entry.Transform.GetFinalValue();
-		Entry.Transform.CreateWaveTable(TransformWaveTable, FinalValue, bBipolar);
-		Entry.Transform.SetFinalValue(FinalValue);
+	PRAGMA_DISABLE_DEPRECATION_WARNINGS
+	ProxyData = MakeShared<FWaveTableBankAssetProxy, ESPMode::ThreadSafe>(GetUniqueID(), SampleMode, SampleRate, GetEntries());
+	PRAGMA_ENABLE_DEPRECATION_WARNINGS
+}
 
-		WaveTableSizeMB += sizeof(float) * WaveTableSize;
-		WaveTableLengthSec = FMath::Max(WaveTableLengthSec, WaveTableSize / 48000.f);
-	}
+void UWaveTableBank::MoveToProxyData()
+{
+	checkf(!ProxyData.IsValid(), TEXT("Cannot overwrite pre-existing proxy when calling UWaveTableBank::MoveToProxyData"));
 
-	WaveTableSizeMB /= FMath::Square(1024);
+	PRAGMA_DISABLE_DEPRECATION_WARNINGS
+	ProxyData = MakeShared<FWaveTableBankAssetProxy, ESPMode::ThreadSafe>(GetUniqueID(), SampleMode, SampleRate, MoveTemp(Entries));
+	PRAGMA_ENABLE_DEPRECATION_WARNINGS
 }
 
+TArray<FWaveTableBankEntry>& UWaveTableBank::GetEntries()
+{
+	PRAGMA_DISABLE_DEPRECATION_WARNINGS
+	return Entries;
+	PRAGMA_ENABLE_DEPRECATION_WARNINGS
+}
+
+#if WITH_EDITOR
 void UWaveTableBank::PostEditChangeProperty(FPropertyChangedEvent& InPropertyChangedEvent)
 {
 	FProperty* Property = InPropertyChangedEvent.Property;
 	if (Property && InPropertyChangedEvent.ChangeType != EPropertyChangeType::Interactive)
 	{
+		// For larger banks, this is very expensive and slows the UX considerably.  This should be
+		// distilled down to array entry if possible for edits on individual WaveTableSettings for
+		// example to make UX more snappy and improve overall property chain edit perf.
 		RefreshWaveTables();
 	}
 
@@ -60,6 +198,9 @@ void UWaveTableBank::PostEditChangeChainProperty(FPropertyChangedChainEvent& InP
 	FProperty* Property = InPropertyChangedEvent.Property;
 	if (Property && InPropertyChangedEvent.ChangeType != EPropertyChangeType::Interactive)
 	{
+		// For larger banks, this is very expensive and slows the UX considerably.  This should be
+		// distilled down to array entry if possible for edits on individual WaveTableSettings for
+		// example to make UX more snappy and improve overall property chain edit perf.
 		RefreshWaveTables();
 	}
 
@@ -75,5 +216,92 @@ void UWaveTableBank::PreSave(FObjectPreSaveContext InSaveContext)
 
 	Super::PreSave(InSaveContext);
 }
+
+void UWaveTableBank::RefreshWaveTables()
+{
+	PRAGMA_DISABLE_DEPRECATION_WARNINGS
+	WaveTableSizeMB = 0.0f;
+	for (FWaveTableBankEntry& Entry : Entries)
+	{
+		FWaveTableData& RuntimeData = Entry.Transform.TableData;
+		const int32 NumSamples = ::WaveTable::BankPrivate::GetWaveTableNumSamples(*this, Entry);
+
+		RuntimeData = FWaveTableData(Entry.Transform.WaveTableSettings.SourceData.GetBitDepth());
+		RuntimeData.Zero(NumSamples);
+
+		Entry.Transform.CreateWaveTable(RuntimeData, bBipolar);
+		if (Entry.Transform.Curve == EWaveTableCurve::File && SampleMode == EWaveTableSamplingMode::FixedSampleRate)
+		{
+			Entry.Transform.Duration = NumSamples / (float)SampleRate;
+		}
+
+		WaveTableSizeMB += RuntimeData.GetRawData().Num();
+	}
+	PRAGMA_ENABLE_DEPRECATION_WARNINGS
+
+	WaveTableSizeMB /= FMath::Square(1024);
+
+	ProxyData.Reset();
+}
 #endif // WITH_EDITOR
-#undef LOCTEXT_NAMESPACE // WaveTable
+
+void UWaveTableBank::Serialize(FArchive& Ar)
+{
+	Super::Serialize(Ar);
+
+#if WITH_EDITORONLY_DATA
+	if (Ar.IsLoading())
+	{
+		PRAGMA_DISABLE_DEPRECATION_WARNINGS
+		for (FWaveTableBankEntry& Entry : Entries)
+		{
+			Entry.Transform.VersionTableData();
+		}
+		PRAGMA_ENABLE_DEPRECATION_WARNINGS
+	}
+#endif // WITH_EDITORONLY_DATA
+}
+
+FWaveTableBankAssetProxy::FWaveTableBankAssetProxy(const FWaveTableBankAssetProxy& InAssetProxy)
+	: ObjectId(InAssetProxy.ObjectId)
+	, SampleRate(InAssetProxy.SampleRate)
+	, SampleMode(InAssetProxy.SampleMode)
+	, WaveTableData(InAssetProxy.WaveTableData)
+{
+}
+
+FWaveTableBankAssetProxy::FWaveTableBankAssetProxy(uint32 InObjectId, EWaveTableSamplingMode InSamplingMode, int32 InSampleRate, const TArray<FWaveTableBankEntry>& InBankEntries)
+{
+	SampleMode = InSamplingMode;
+	SampleRate = InSampleRate;
+	Algo::Transform(InBankEntries, WaveTableData, [](const FWaveTableBankEntry& Entry)
+	{
+		return Entry.Transform.GetTableData();
+	});
+}
+
+FWaveTableBankAssetProxy::FWaveTableBankAssetProxy(uint32 InObjectId, EWaveTableSamplingMode InSamplingMode, int32 InSampleRate, TArray<FWaveTableBankEntry>&& InBankEntries)
+{
+	SampleMode = InSamplingMode;
+	SampleRate = InSampleRate;
+	for (FWaveTableBankEntry& Entry : InBankEntries)
+	{
+		WaveTableData.Add(MoveTemp(Entry.Transform.TableData));
+	}
+	InBankEntries.Empty();
+}
+
+FWaveTableBankAssetProxy::FWaveTableBankAssetProxy(const UWaveTableBank& InWaveTableBank)
+{
+	ObjectId = InWaveTableBank.GetUniqueID();
+	SampleRate = InWaveTableBank.SampleRate;
+	SampleMode = InWaveTableBank.SampleMode;
+
+	PRAGMA_DISABLE_DEPRECATION_WARNINGS
+	Algo::Transform(InWaveTableBank.Entries, WaveTableData, [](const FWaveTableBankEntry& Entry)
+	{
+		return Entry.Transform.GetTableData();
+	});
+	PRAGMA_ENABLE_DEPRECATION_WARNINGS
+
+}
diff --git a/Engine/Plugins/Runtime/WaveTable/Source/WaveTable/Private/WaveTableImporter.cpp b/Engine/Plugins/Runtime/WaveTable/Source/WaveTable/Private/WaveTableImporter.cpp
--- a/Engine/Plugins/Runtime/WaveTable/Source/WaveTable/Private/WaveTableImporter.cpp
+++ b/Engine/Plugins/Runtime/WaveTable/Source/WaveTable/Private/WaveTableImporter.cpp
@@ -15,23 +15,15 @@ namespace WaveTable
 		Sampler.SetPhase(Settings.Phase);
 	}
 
-	void FImporter::Process(TArray<float>& OutWaveTable)
+	void FImporter::ProcessInternal(const TArrayView<const float> InEditSourceView, TArray<float>& OutWaveTable)
 	{
-		// 1. Get editable section of source PCM
-		const TArrayView<const float> View = Settings.GetEditSourceView();
+		const FWaveTableView EditSourceView = FWaveTableView(InEditSourceView, InEditSourceView.Last());
 
-		if(View.IsEmpty())
-		{
-			return;
-		}
-		
-		const FWaveTableView EditSourceView = FWaveTableView(View, View.Last());
-		
-		// 2. Resample into table
+		// 1. Resample into table
 		Sampler.Reset();
 		Sampler.Process(EditSourceView, OutWaveTable);
 
-		// 3. Apply offset
+		// 2. Apply offset
 		float TableOffset = 0.0f;
 		if (Settings.bRemoveOffset || !bBipolar)
 		{
@@ -39,7 +31,7 @@ namespace WaveTable
 			Audio::ArrayAddConstantInplace(OutWaveTable, -1.0f * TableOffset);
 		}
 
-		// 4. Normalize
+		// 3. Normalize
 		if (Settings.bNormalize)
 		{
 			const float MaxValue = Audio::ArrayMaxAbsValue(OutWaveTable);
@@ -49,7 +41,7 @@ namespace WaveTable
 			}
 		}
 
-		// 5. Apply fades
+		// 4. Apply fades
 		if (Settings.FadeIn > 0.0f)
 		{
 			const int32 FadeLength = RatioToIndex(Settings.FadeIn, OutWaveTable.Num());
@@ -65,7 +57,7 @@ namespace WaveTable
 			Audio::ArrayFade(FadeView, 1.0f, 0.0f);
 		}
 
-		// 6. Finalize if unipolar source
+		// 5. Finalize if unipolar source
 		if (!bBipolar)
 		{
 			Audio::ArrayAbsInPlace(OutWaveTable);
@@ -77,5 +69,97 @@ namespace WaveTable
 				Audio::ArrayAddConstantInplace(OutWaveTable, TableOffset);
 			}
 		}
+	}
+
+	void FImporter::Process(TArray<float>& OutWaveTable)
+	{
+		PRAGMA_DISABLE_DEPRECATION_WARNINGS
+		const TArrayView<const float> View = Settings.GetEditSourceView();
+		if (!View.IsEmpty())
+		{
+			ProcessInternal(View, OutWaveTable);
+		}
+		PRAGMA_ENABLE_DEPRECATION_WARNINGS
+	}
+
+	void FImporter::Process(::FWaveTableData& OutData)
+	{
+		// 1. Get editable section of source and run BDC to process light edits on float data
+		int32 Offset = 0;
+		int32 NumEditSamples = 0;
+		Settings.GetEditSourceBounds(Offset, NumEditSamples);
+
+		if (NumEditSamples == 0)
+		{
+			OutData.Zero();
+			return;
+		}
+
+		// 2. Get source in float bit depth format
+		TArrayView<const float> EditDataView;
+		TArray<float> ScratchData;
+		switch (Settings.SourceData.GetBitDepth())
+		{
+			case EWaveTableBitDepth::IEEE_Float:
+			{
+				Settings.SourceData.GetDataView(EditDataView);
+				EditDataView = { EditDataView.GetData() + Offset, NumEditSamples };
+			}
+			break;
+
+			case EWaveTableBitDepth::PCM_16:
+			{
+				TArrayView<const int16> EditDataView16Bit;
+				Settings.SourceData.GetDataView(EditDataView16Bit);
+				EditDataView16Bit = { EditDataView16Bit.GetData() + Offset, NumEditSamples };
+				ScratchData.AddZeroed(NumEditSamples);
+				Audio::ArrayPcm16ToFloat(EditDataView16Bit, ScratchData);
+				EditDataView = ScratchData;
+			}
+			break;
+
+			default:
+			{
+				static_assert(static_cast<int32>(EWaveTableBitDepth::COUNT) == 2, "Possible missing switch case coverage for 'EWaveTableBitDepth'");
+				checkNoEntry();
+			}
+			break;
+		}
+
+		// 3. Process effects
+		TArray<float> ProcessedData;
+		ProcessedData.AddZeroed(OutData.GetNumSamples());
+		ProcessInternal(EditDataView, ProcessedData);
+
+		// 4. Copy processed data into original buffer
+		switch (OutData.GetBitDepth())
+		{
+			case EWaveTableBitDepth::IEEE_Float:
+			{
+				constexpr bool bIsLoop = true;
+				OutData.SetData(ProcessedData, bIsLoop);
+			}
+			break;
+
+			case EWaveTableBitDepth::PCM_16:
+			{
+				TArrayView<int16> TableView;
+				OutData.GetDataView(TableView);
+				Audio::ArrayFloatToPcm16(ProcessedData, TableView);
+				if (!TableView.IsEmpty())
+				{
+					constexpr float ConversionValue = 1.0f / static_cast<float>(TNumericLimits<int16>::Max());
+					OutData.SetFinalValue(TableView[0] * ConversionValue);
+				}
+			}
+			break;
+
+			default:
+			{
+				static_assert(static_cast<int32>(EWaveTableBitDepth::COUNT) == 2, "Possible missing switch case coverage for 'EWaveTableBitDepth'");
+				checkNoEntry();
+			}
+			break;
+		}
 	} // namespace Importer
 } // namespace WaveTable
diff --git a/Engine/Plugins/Runtime/WaveTable/Source/WaveTable/Private/WaveTableSampler.cpp b/Engine/Plugins/Runtime/WaveTable/Source/WaveTable/Private/WaveTableSampler.cpp
--- a/Engine/Plugins/Runtime/WaveTable/Source/WaveTable/Private/WaveTableSampler.cpp
+++ b/Engine/Plugins/Runtime/WaveTable/Source/WaveTable/Private/WaveTableSampler.cpp
@@ -10,26 +10,29 @@ namespace WaveTable
 {
 	namespace UtilitiesPrivate
 	{
-		UE_NODISCARD FORCEINLINE FVector2f GetTangentP0(const float* InTableView, const float NextValue, float InIndex, int32 InArraySize)
+		template<typename TSample>
+		FVector2f GetTangentP0(const TSample* InTableView, const float NextValue, float InIndex, int32 InArraySize)
 		{
 			const int32 Index = FMath::TruncToInt32(InIndex);
 
 			// Addition of array size guarantees that if index is 0, we aren't modding a negative number and returning a negative index
-			const float LastValue = InTableView[(Index - 1 + InArraySize) % InArraySize];
-			return FVector2f(2.0f, NextValue - LastValue).GetSafeNormal();
+			const TSample LastValue = InTableView[(Index - 1 + InArraySize) % InArraySize];
+			return FVector2f(2.0f, (float)(NextValue - LastValue)).GetSafeNormal();
 		};
 
-		UE_NODISCARD FORCEINLINE FVector2f GetTangentP1(const float* InTableView, const float LastValue, float InIndex, int32 InArraySize)
+		template<typename TSample>
+		FVector2f GetTangentP1(const TSample* InTableView, const float LastValue, float InIndex, int32 InArraySize)
 		{
 			const int32 Index = FMath::TruncToInt32(InIndex);
-			const float NextValue = InTableView[(Index + 1) % InArraySize];
-			return FVector2f(2.0f, NextValue - LastValue).GetSafeNormal();
+			const TSample NextValue = InTableView[(Index + 1) % InArraySize];
+			return FVector2f(2.0f, (float)(NextValue - LastValue)).GetSafeNormal();
 		};
 
-		auto MaxValueIndexInterpolator = [](TArrayView<const float> InTableView, TArrayView<float> InOutIndicesToValues)
+		template<typename TSample>
+		void MaxValueIndexInterpolator(TArrayView<const TSample> InTableView, TArrayView<float> InOutIndicesToValues, float InGain = 1.0f)
 		{
 			int32 LastIndex = 0;
-			const float* TableData = InTableView.GetData();
+			const TSample* TableData = InTableView.GetData();
 			float* IndexValueData = InOutIndicesToValues.GetData();
 			for (int32 i = 0; i < InOutIndicesToValues.Num(); ++i)
 			{
@@ -41,14 +44,14 @@ namespace WaveTable
 					::Swap(LastIndex, Index);
 				}
 
-				int32 MaxAbs = 0.0f;
+				TSample MaxAbs = 0.0f;
 				for (int32 j = LastIndex; j <= Index; ++j)
 				{
-					float IndexAbs = FMath::Abs(TableData[j]);
+					TSample IndexAbs = FMath::Abs(TableData[j]);
 					if (IndexAbs > MaxAbs)
 					{
 						MaxAbs = IndexAbs;
-						IndexValueData[i] = TableData[j];
+						IndexValueData[i] = InGain * (float)TableData[j];
 					}
 				}
 
@@ -59,54 +62,57 @@ namespace WaveTable
 			}
 		};
 
-		auto CubicIndexInterpolator = [](TArrayView<const float> InTableView, TArrayView<float> InOutIndicesToValues)
+		template <typename TSample>
+		void CubicIndexInterpolator (TArrayView<const TSample> InTableView, TArrayView<float> InOutIndicesToValues, float InGain = 1.0f)
 		{
 			const int32 NumTableSamples = InTableView.Num();
-			const float* InTable = InTableView.GetData();
+			const TSample* InTable = InTableView.GetData();
 			for (int32 i = 0; i < InOutIndicesToValues.Num(); ++i)
 			{
 				float& IndexToOutput = InOutIndicesToValues[i];
 				const int32 LastIndexInt = FMath::TruncToInt32(IndexToOutput);
 
-				const FVector2f P0 = { (float)LastIndexInt, InTable[LastIndexInt % NumTableSamples] };
-				const FVector2f P1 = { (float)(LastIndexInt + 1), InTable[(LastIndexInt + 1) % NumTableSamples] };
+				const FVector2f P0 = { (float)LastIndexInt, (float)InTable[LastIndexInt % NumTableSamples] };
+				const FVector2f P1 = { (float)(LastIndexInt + 1), (float)InTable[(LastIndexInt + 1) % NumTableSamples] };
 				const FVector2f TangentP0 = GetTangentP0(InTable, P1.Y, IndexToOutput, NumTableSamples);
 				const FVector2f TangentP1 = GetTangentP1(InTable, P0.Y, IndexToOutput, NumTableSamples);
 				const float Alpha = IndexToOutput - LastIndexInt;
 
-				IndexToOutput = FMath::CubicInterp(P0, TangentP0, P1, TangentP1, Alpha).Y;
+				IndexToOutput = FMath::CubicInterp(P0, TangentP0, P1, TangentP1, Alpha).Y * InGain;
 			}
 		};
 
-		auto LinearIndexInterpolator = [](TArrayView<const float> InTableView, const TArrayView<float> InOutIndicesToValues)
+		template<typename TSample>
+		void LinearIndexInterpolator(TArrayView<const TSample> InTableView, const TArrayView<float> InOutIndicesToValues, float InGain = 1.0f)
 		{
 			const int32 NumTableSamples = InTableView.Num();
-			const float* InTable = InTableView.GetData();
+			const TSample* InTable = InTableView.GetData();
 			for (int32 i = 0; i < InOutIndicesToValues.Num(); ++i)
 			{
 				float& IndexToOutput = InOutIndicesToValues[i];
 				const int32 LastIndexInt = FMath::TruncToInt32(IndexToOutput);
 
-				const float P0 = InTable[LastIndexInt % NumTableSamples];
+				const float P0 = (float)InTable[LastIndexInt % NumTableSamples];
 				const float Alpha = IndexToOutput - LastIndexInt;
-				const float P1 = InTable[(LastIndexInt + 1) % NumTableSamples];
+				const float P1 = (float)InTable[(LastIndexInt + 1) % NumTableSamples];
 
-				IndexToOutput = FMath::Lerp(P0, P1, Alpha);
+				IndexToOutput = FMath::Lerp(P0, P1, Alpha) * InGain;
 			}
 		};
 
-		auto StepIndexInterpolator = [](TArrayView<const float> InTableView, const TArrayView<float> InOutIndicesToValues)
+		template<typename TSample>
+		void StepIndexInterpolator(TArrayView<const TSample> InTableView, const TArrayView<float> InOutIndicesToValues, float InGain = 1.0f)
 		{
 			const int32 NumTableSamples = InTableView.Num();
-			const float* InTable = InTableView.GetData();
+			const TSample* InTable = InTableView.GetData();
 			float* IndexToValue = InOutIndicesToValues.GetData();
 			for (int32 i = 0; i < InOutIndicesToValues.Num(); ++i)
 			{
 				float& IndexToOutput = IndexToValue[i];
 				const int32 LastIndexInt = FMath::TruncToInt32(IndexToOutput);
-				IndexToOutput = InTable[LastIndexInt % NumTableSamples];
+				IndexToOutput = InTable[LastIndexInt % NumTableSamples] * InGain;
 			}
-		};
+		}
 	}
 
 	FWaveTableSampler::FWaveTableSampler()
@@ -118,6 +124,74 @@ namespace WaveTable
 	{
 	}
 
+	void FWaveTableSampler::Interpolate(const FWaveTableData& InTableData, TArrayView<float> InOutIndexToSamplesView, EInterpolationMode InterpMode)
+	{
+		switch (InTableData.GetBitDepth())
+		{
+			case EWaveTableBitDepth::IEEE_Float:
+			{
+				TArrayView<const float> InTableView;
+				InTableData.GetDataView(InTableView);
+				Interpolate(InTableView, InOutIndexToSamplesView, InterpMode);
+			}
+			break;
+
+			case EWaveTableBitDepth::PCM_16:
+			{
+				TArrayView<const int16> InTableView;
+				InTableData.GetDataView(InTableView);
+				Interpolate(InTableView, InOutIndexToSamplesView, InterpMode);
+			}
+			break;
+
+			default:
+			{
+				static_assert(static_cast<int32>(EWaveTableBitDepth::COUNT) == 2, "Possible missing switch case coverage for 'EWaveTableBitDepth'");
+				checkNoEntry();
+			}
+		}
+	}
+
+	void FWaveTableSampler::Interpolate(TArrayView<const int16> InTableView, TArrayView<float> InOutIndexToSamplesView, EInterpolationMode InterpMode)
+	{
+		using namespace UtilitiesPrivate;
+
+		constexpr float ConversionValue = 1.0f / static_cast<float>(TNumericLimits<int16>::Max());
+		switch (InterpMode)
+		{
+			case EInterpolationMode::None:
+			{
+				StepIndexInterpolator(InTableView, InOutIndexToSamplesView, ConversionValue);
+			}
+			break;
+
+			case EInterpolationMode::Linear:
+			{
+				LinearIndexInterpolator(InTableView, InOutIndexToSamplesView, ConversionValue);
+			}
+			break;
+
+			case EInterpolationMode::Cubic:
+			{
+				CubicIndexInterpolator(InTableView, InOutIndexToSamplesView, ConversionValue);
+			}
+			break;
+
+			case EInterpolationMode::MaxValue:
+			{
+				MaxValueIndexInterpolator(InTableView, InOutIndexToSamplesView, ConversionValue);
+			}
+			break;
+
+			default:
+			{
+				static_assert(static_cast<int32>(EInterpolationMode::COUNT) == 4, "Possible missing switch coverage for EInterpolationMode");
+				checkNoEntry();
+			}
+			break;
+		}
+	}
+
 	void FWaveTableSampler::Interpolate(TArrayView<const float> InTableView, TArrayView<float> InOutIndexToSamplesView, EInterpolationMode InterpMode)
 	{
 		using namespace UtilitiesPrivate;
@@ -177,6 +251,11 @@ namespace WaveTable
 		Settings.InterpolationMode = InMode;
 	}
 
+	void FWaveTableSampler::SetOneShot(bool bInOneShot)
+	{
+		Settings.bOneShot = bInOneShot;
+	}
+
 	void FWaveTableSampler::SetPhase(float InPhase)
 	{
 		Settings.Phase = InPhase;
@@ -203,7 +282,7 @@ namespace WaveTable
 		// to see if interpolation needs to be re-evaluated.
 		// This check is expensive in the context of buffer
 		// eval and thus only supported in "single-sample" mode.
-		const bool bIsFinalInterp = Index > InTableView.SampleView.Num() - 1;
+		const bool bIsFinalInterp = Index > InTableView.Num() - 1;
 		const bool bHasCompleted = Index == 0 && (Settings.Phase > 0.0f || LastIndex > 0);
 		if (bIsFinalInterp || bHasCompleted)
 		{
@@ -270,19 +349,20 @@ namespace WaveTable
 
 		if (!OutSamplesView.IsEmpty())
 		{
-			if (InTableView.SampleView.IsEmpty())
+			if (InTableView.IsEmpty())
 			{
 				Audio::ArraySetToConstantInplace(OutSamplesView, 0.0f);
 			}
 			else
 			{
-				// Compute index frequency & phase and stuff into samples
-				// view. Interpolator will convert to associated values.
-				ComputeIndexFrequency(InTableView.SampleView, InFreqModulator, InSyncTriggers, OutSamplesView);
-				ComputeIndexPhase(InTableView.SampleView, InPhaseModulator, OutSamplesView);
+				// Compute index frequency & phase and stuff into samples view.
+				// Interpolator will convert to associated values.
+				ComputeIndexFrequency(InTableView.Num(), InFreqModulator, InSyncTriggers, OutSamplesView);
+				ComputeIndexPhase(InTableView.Num(), InPhaseModulator, OutSamplesView);
 
-				// Capture last index and return in case caller is interested in progress/phase through table (ex. for enveloping).
-				CurViewLastIndex = OutSamplesView.Last();
+				// Find stopping index if it is found and sampler is set to one shot,
+				// and if set clear out section of samples view beyond end computed in prior steps.
+				CurViewLastIndex = ComputeIndexFinished(InSyncTriggers, OutSamplesView);
 
 				Interpolate(InTableView.SampleView, OutSamplesView, Settings.InterpolationMode);
 			}
@@ -301,25 +381,91 @@ namespace WaveTable
 		return CurViewLastIndex;
 	}
 
-	void FWaveTableSampler::ComputeIndexFrequency(TArrayView<const float> InTableView, TArrayView<const float> InFreqModulator, TArrayView<const float> InSyncTriggers, TArrayView<float> OutIndicesView)
+	float FWaveTableSampler::Process(const FWaveTableData& InTableData, TArrayView<float> OutSamplesView)
+	{
+		return Process(InTableData, { }, { }, { }, OutSamplesView);
+	}
+
+	float FWaveTableSampler::Process(const FWaveTableData& InTableData, TArrayView<const float> InFreqModulator, TArrayView<const float> InPhaseModulator, TArrayView<const float> InSyncTriggers, TArrayView<float> OutSamplesView)
+	{
+		float CurViewLastIndex = 0.0f;
+
+		if (!OutSamplesView.IsEmpty())
+		{
+			if (InTableData.IsEmpty())
+			{
+				Audio::ArraySetToConstantInplace(OutSamplesView, 0.0f);
+			}
+			else
+			{
+				// Compute index frequency & phase and stuff into samples view.
+				// Interpolator will convert to associated values.
+				ComputeIndexFrequency(InTableData.GetNumSamples(), InFreqModulator, InSyncTriggers, OutSamplesView);
+				ComputeIndexPhase(InTableData.GetNumSamples(), InPhaseModulator, OutSamplesView);
+
+				// Find stopping index if it is found and sampler is set to one shot,
+				// and if set clear out section of samples view beyond end computed in prior steps.
+				CurViewLastIndex = ComputeIndexFinished(InSyncTriggers, OutSamplesView);
+
+				switch (InTableData.GetBitDepth())
+				{
+					case EWaveTableBitDepth::IEEE_Float:
+					{
+						TArrayView<const float> DataView;
+						InTableData.GetDataView(DataView);
+						Interpolate(DataView, OutSamplesView, Settings.InterpolationMode);
+					}
+					break;
+
+					case EWaveTableBitDepth::PCM_16:
+					{
+						TArrayView<const int16> DataView;
+						InTableData.GetDataView(DataView);
+						Interpolate(DataView, OutSamplesView, Settings.InterpolationMode);
+					}
+					break;
+
+					default:
+					{
+						static_assert(static_cast<int32>(EWaveTableBitDepth::COUNT) == 2, "Possible missing switch case coverage for 'EWaveTableBitDepth'");
+						checkNoEntry();
+					}
+				}
+			}
+
+			if (!FMath::IsNearlyEqual(Settings.Amplitude, 1.0f))
+			{
+				Audio::ArrayMultiplyByConstantInPlace(OutSamplesView, Settings.Amplitude);
+			}
+
+			if (!FMath::IsNearlyZero(Settings.Offset))
+			{
+				Audio::ArrayAddConstantInplace(OutSamplesView, Settings.Offset);
+			}
+		}
+
+		return CurViewLastIndex;
+	}
+
+	void FWaveTableSampler::ComputeIndexFrequency(int32 NumInputSamples, TArrayView<const float> InFreqModulator, TArrayView<const float> InSyncTriggers, TArrayView<float> OutIndicesView)
 	{
-		Audio::ArraySetToConstantInplace(OutIndicesView, Settings.Freq * InTableView.Num() / OutIndicesView.Num());
+		Audio::ArraySetToConstantInplace(OutIndicesView, Settings.Freq * NumInputSamples / OutIndicesView.Num());
 
 		if (InFreqModulator.Num() == OutIndicesView.Num())
 		{
 			Audio::ArrayMultiplyInPlace(InFreqModulator, OutIndicesView);
 		}
 
-		auto TransformLastIndex = [this, InTableView, &InFreqModulator, &InSyncTriggers, &OutIndicesView](float InputIndex)
+		auto TransformLastIndex = [this, NumInputSamples, &InFreqModulator, &InSyncTriggers, &OutIndicesView](float InputIndex)
 		{
 			if (InFreqModulator.Num() == OutIndicesView.Num())
 			{
-				return (InFreqModulator.Last() * Settings.Freq * InTableView.Num()) + InputIndex;
+				return (InFreqModulator.Last() * Settings.Freq * NumInputSamples) + InputIndex;
 			}
 			else
 			{
 				checkf(InFreqModulator.IsEmpty(), TEXT("FreqModulator view should be the same size as the sample view or not supplied (size of 0)."));
-				return (Settings.Freq * InTableView.Num()) + InputIndex;
+				return (Settings.Freq * NumInputSamples) + InputIndex;
 			}
 		};
 
@@ -333,7 +479,7 @@ namespace WaveTable
 				// Numerically negate trig value as high (1.0f) should
 				// be zero to "reset" index state (avoids conditional
 				// in loop at expense of memory).
-				const float InvertedTrig = -1 * SyncTrigData[i] + 1;
+				const float InvertedTrig = -1.0f * SyncTrigData[i] + 1.0f;
 				SyncIndex *= InvertedTrig;
 				LastIndex *= InvertedTrig;
 				OutSamples[i] *= SyncIndex;
@@ -356,10 +502,10 @@ namespace WaveTable
 			LastIndex = TransformLastIndex(LastIndex);
 		}
 
-		LastIndex = FMath::Frac(LastIndex / InTableView.Num()) * InTableView.Num();
+		LastIndex = FMath::Frac(LastIndex / NumInputSamples) * NumInputSamples;
 	}
 
-	void FWaveTableSampler::ComputeIndexPhase(TArrayView<const float> InTableView, TArrayView<const float> InPhaseModulator, TArrayView<float> OutIndicesView)
+	void FWaveTableSampler::ComputeIndexPhase(int32 NumInputSamples, TArrayView<const float> InPhaseModulator, TArrayView<float> OutIndicesView)
 	{
 		check(!OutIndicesView.IsEmpty());
 
@@ -373,7 +519,7 @@ namespace WaveTable
 			float* OutSamples = OutIndicesView.GetData();
 			for (int32 i = 0; i < OutIndicesView.Num(); ++i)
 			{
-				const float PhaseIndexOffset = FMath::Frac(PhaseModData[i]) * InTableView.Num();
+				const float PhaseIndexOffset = FMath::Frac(PhaseModData[i]) * NumInputSamples;
 				OutSamples[i] += PhaseIndexOffset;
 			}
 		}
@@ -381,8 +527,63 @@ namespace WaveTable
 		{
 			checkf(InPhaseModulator.IsEmpty(), TEXT("PhaseModulator view should be the same size as the sample view or not supplied (size of 0)."));
 
-			const float PhaseIndexOffset = FMath::Frac(Settings.Phase) * InTableView.Num();
+			const float PhaseIndexOffset = FMath::Frac(Settings.Phase) * NumInputSamples;
 			Audio::ArrayAddConstantInplace(OutIndicesView, PhaseIndexOffset);
 		}
 	}
+
+	int32 FWaveTableSampler::ComputeIndexFinished(TArrayView<const float> InSyncTriggers, TArrayView<float> OutIndicesView)
+	{
+		check(!OutIndicesView.IsEmpty());
+
+		if (Settings.bOneShot)
+		{
+			bool Stopped = false;
+			OneShotData.IndexFinished = 0;
+			if (InSyncTriggers.IsEmpty())
+			{
+				for (int32 ArrayIndex = 0; ArrayIndex < OutIndicesView.Num(); ++ArrayIndex)
+				{
+					const float& ViewIndex = OutIndicesView[ArrayIndex];
+					if (OneShotData.LastOutputIndex >= ViewIndex)
+					{
+						OneShotData.IndexFinished = ArrayIndex;
+						return OneShotData.IndexFinished;
+					}
+
+					OneShotData.LastOutputIndex = ViewIndex;
+				}
+			}
+			else
+			{
+				checkf(InSyncTriggers.Num() == OutIndicesView.Num(), TEXT("If SyncTriggers view is supplied, must match length of given indices view"))
+				for (int32 ArrayIndex = 0; ArrayIndex < OutIndicesView.Num(); ++ArrayIndex)
+				{
+					const float& ViewIndex = OutIndicesView[ArrayIndex];
+					const bool NotSyncing = InSyncTriggers[ArrayIndex] == 0.0f;
+					if (NotSyncing && OneShotData.LastOutputIndex >= ViewIndex)
+					{
+						OneShotData.IndexFinished = ArrayIndex;
+						return OneShotData.IndexFinished;
+					}
+
+					OneShotData.LastOutputIndex = ViewIndex;
+				}
+			}
+
+			if (OutIndicesView.IsEmpty() || Stopped != 0)
+			{
+				OneShotData.LastOutputIndex = -1.0f;
+				return OneShotData.IndexFinished;
+			}
+		}
+
+		OneShotData.IndexFinished = INDEX_NONE;
+		return OutIndicesView.Last();
+	}
+
+	int32 FWaveTableSampler::GetIndexFinished() const
+	{
+		return OneShotData.IndexFinished;
+	}
 } // namespace WaveTable
diff --git a/Engine/Plugins/Runtime/WaveTable/Source/WaveTable/Private/WaveTableSettings.cpp b/Engine/Plugins/Runtime/WaveTable/Source/WaveTable/Private/WaveTableSettings.cpp
--- a/Engine/Plugins/Runtime/WaveTable/Source/WaveTable/Private/WaveTableSettings.cpp
+++ b/Engine/Plugins/Runtime/WaveTable/Source/WaveTable/Private/WaveTableSettings.cpp
@@ -3,9 +3,41 @@
 
 #include UE_INLINE_GENERATED_CPP_BY_NAME(WaveTableSettings)
 
+void FWaveTableSettings::GetEditSourceBounds(int32& OutSourceTopOffset, int32& OutSourceNumSamples) const
+{
+	OutSourceTopOffset = 0;
+	OutSourceNumSamples = 0;
+
+	int32 SourceTopOffset = 0;
+	int32 SourceNumSamples = SourceData.GetNumSamples();
+	if (SourceNumSamples == 0)
+	{
+		return;
+	}
+
+	if (Top > 0.0f)
+	{
+		SourceTopOffset = WaveTable::RatioToIndex(Top, SourceNumSamples);
+		SourceNumSamples -= SourceTopOffset;
+	}
+
+	if (Tail > 0.0f)
+	{
+		const int32 SourceTailOffset = WaveTable::RatioToIndex(Tail, SourceNumSamples);
+		SourceNumSamples -= SourceTailOffset;
+	}
+
+	if (SourceTopOffset < SourceData.GetNumSamples())
+	{
+		OutSourceTopOffset = SourceTopOffset;
+		OutSourceNumSamples = SourceNumSamples;
+	}
+}
 
 TArrayView<const float> FWaveTableSettings::GetEditSourceView() const
 {
+#if WITH_EDITOR
+	PRAGMA_DISABLE_DEPRECATION_WARNINGS
 	int32 SourceNumSamples = SourcePCMData.Num();
 	if (SourceNumSamples == 0)
 	{
@@ -29,74 +61,30 @@ TArrayView<const float> FWaveTableSettings::GetEditSourceView() const
 	{
 		return { SourcePCMData.GetData() + SourceTopOffset, SourceNumSamples };
 	}
+	PRAGMA_ENABLE_DEPRECATION_WARNINGS
 
+#endif // WITH_EDITOR
 	return { };
 }
 
+#if WITH_EDITOR
+void FWaveTableSettings::VersionPCMData()
+{
+	PRAGMA_DISABLE_DEPRECATION_WARNINGS
+	if (!SourcePCMData.IsEmpty())
+	{
+		SourceData.SetRawData(EWaveTableBitDepth::IEEE_Float, TArrayView<uint8>((uint8*)(SourcePCMData.GetData()), SourcePCMData.Num() * sizeof(float)));
+		SourcePCMData.Empty();
+	}
+	PRAGMA_ENABLE_DEPRECATION_WARNINGS
+}
+#endif // WITH_EDITOR
+
 namespace WaveTable
 {
 	int32 GetWaveTableSize(EWaveTableResolution InWaveTableResolution, EWaveTableCurve InCurve, int32 InMaxPCMSize)
 	{
-		auto GetCurveResolution = [](EWaveTableCurve Curve)
-		{
-			switch (Curve)
-			{
-				case EWaveTableCurve::File:
-				{
-					return EWaveTableResolution::Maximum;
-				}
-				break;
-
-				case EWaveTableCurve::Linear:
-				case EWaveTableCurve::Linear_Inv:
-				{
-					return EWaveTableResolution::Res_8;
-				}
-				break;
-
-				case EWaveTableCurve::Exp:
-				case EWaveTableCurve::Exp_Inverse:
-				case EWaveTableCurve::Log:
-				{
-					return EWaveTableResolution::Res_256;
-				}
-				break;
-
-				case EWaveTableCurve::Sin:
-				case EWaveTableCurve::SCurve:
-				case EWaveTableCurve::Sin_Full:
-				{
-					return EWaveTableResolution::Res_64;
-				}
-				break;
-
-				default:
-				{
-					return EWaveTableResolution::Res_128;
-				}
-				break;
-			};
-		};
-
-		switch (InWaveTableResolution)
-		{
-			case EWaveTableResolution::Maximum:
-			{
-				return FMath::Max(InMaxPCMSize, 1 << static_cast<int32>(EWaveTableResolution::Res_Max));
-			}
-			break;
-
-			case EWaveTableResolution::None:
-			{
-				const EWaveTableResolution CurveRes = GetCurveResolution(InCurve);
-				return 1 << static_cast<int32>(CurveRes);
-			}
-			break;
-
-			default:
-			{
-				return 1 << static_cast<int32>(InWaveTableResolution);
-			}
-		};
+		checkNoEntry();
+		return 0;
 	}
 } // namespace WaveTable
diff --git a/Engine/Plugins/Runtime/WaveTable/Source/WaveTable/Private/WaveTableTransform.cpp b/Engine/Plugins/Runtime/WaveTable/Source/WaveTable/Private/WaveTableTransform.cpp
--- a/Engine/Plugins/Runtime/WaveTable/Source/WaveTable/Private/WaveTableTransform.cpp
+++ b/Engine/Plugins/Runtime/WaveTable/Source/WaveTable/Private/WaveTableTransform.cpp
@@ -7,38 +7,50 @@
 
 #include UE_INLINE_GENERATED_CPP_BY_NAME(WaveTableTransform)
 
+
 void FWaveTableTransform::Apply(float& InOutValue, bool bInBipolar) const
 {
-	if (WaveTable.IsEmpty())
+	FWaveTableData TempTable(TArrayView<float>{ &InOutValue, 1 }, 0.0f);
+	TArrayView<float> TempData;
+	TempTable.GetDataView(TempData);
+
+	// If no table data is cached, attempt to sample from curve
+	if (TableData.IsEmpty())
 	{
-		float DummyVal = 0.f;
-		TArrayView<float> ValueView = { &InOutValue, 1 };
-		bInBipolar ? SampleCurveBipolar(ValueView, DummyVal) : SampleCurveUnipolar(ValueView, DummyVal);
+		bInBipolar ? SampleCurveBipolar(TempTable) : SampleCurveUnipolar(TempTable);
 	}
 	else
 	{
-		InOutValue *= WaveTable.Num(); // Convert to float index
-		TArrayView<float> ValueView = { &InOutValue, 1 };
-		WaveTable::FWaveTableSampler::Interpolate(WaveTable, ValueView);
+		InOutValue *= TableData.GetNumSamples(); // Convert to float index
+		TArrayView<const float> TableView;
+		TableData.GetDataView(TableView);
+		WaveTable::FWaveTableSampler::Interpolate(TableView, TempData);
 	}
+
+	InOutValue = TempData.Last();
 }
 
-void FWaveTableTransform::SampleCurveBipolar(TArrayView<float> InOutTable, float& OutFinalValue) const
+void FWaveTableTransform::SampleCurveBipolar(FWaveTableData& InOutTableData) const
 {
+	checkf(InOutTableData.GetBitDepth() == EWaveTableBitDepth::IEEE_Float, TEXT("Curve sampling only supports IEEE 32-bit float bit depth"));
+
+	TArrayView<float> TableView;
+	InOutTableData.GetDataView(TableView);
+
 	// Clamp the input
-	Audio::ArrayClampInPlace(InOutTable, 0.0f, 1.0f);
+	Audio::ArrayClampInPlace(TableView, 0.0f, 1.0f);
 
-	float* ValueData = InOutTable.GetData();
+	float* ValueData = TableView.GetData();
 	switch (Curve)
 	{
 		case EWaveTableCurve::Custom:
 		{
-			for (int32 i = 0; i < InOutTable.Num(); ++i)
+			for (int32 i = 0; i < TableView.Num(); ++i)
 			{
 				ValueData[i] = CurveCustom.Eval(ValueData[i]);
 			}
 				
-			OutFinalValue = CurveCustom.Eval(1.f);
+			InOutTableData.SetFinalValue(CurveCustom.Eval(1.0f));
 			break;
 		}
 
@@ -46,35 +58,35 @@ void FWaveTableTransform::SampleCurveBipolar(TArrayView<float> InOutTable, float
 		{
 			if (CurveShared)
 			{
-				for (int32 i = 0; i < InOutTable.Num(); ++i)
+				for (int32 i = 0; i < TableView.Num(); ++i)
 				{
 					ValueData[i] = CurveShared->FloatCurve.Eval(ValueData[i]);
 				}
 				
-				OutFinalValue = CurveShared->FloatCurve.Eval(1.f);
+				InOutTableData.SetFinalValue(CurveShared->FloatCurve.Eval(1.0f));
 			}
 			break;
 		}
 
 		case EWaveTableCurve::Linear_Inv:
 		{
-			for (int32 i = 0; i < InOutTable.Num(); ++i)
+			for (int32 i = 0; i < TableView.Num(); ++i)
 			{
 				ValueData[i] = FMath::Lerp(1.0f, -1.0f, ValueData[i]);
 			}
 				
-			OutFinalValue = -1.f;
+			InOutTableData.SetFinalValue(-1.0f);
 		}
 		break;
 
 		case EWaveTableCurve::Linear:
 		{
-			for (int32 i = 0; i < InOutTable.Num(); ++i)
+			for (int32 i = 0; i < TableView.Num(); ++i)
 			{
 				ValueData[i] = FMath::Lerp(-1.0f, 1.0f, ValueData[i]);
 			}
 			
-			OutFinalValue = 1.f;
+			InOutTableData.SetFinalValue(1.0f);
 		}
 		break;
 
@@ -82,13 +94,13 @@ void FWaveTableTransform::SampleCurveBipolar(TArrayView<float> InOutTable, float
 		{
 			// Alpha is limited to between 0.0f and 1.0f and ExponentialScalar
 			// between 0 and 10 to keep ValueData "sane" and avoid float boundary.
-			for (int32 i = 0; i < InOutTable.Num(); ++i)
+			for (int32 i = 0; i < TableView.Num(); ++i)
 			{
 				ValueData[i] *= (FMath::Pow(10.0f, Scalar * (ValueData[i] - 1.0f)));
 				ValueData[i] -= 0.5f;
 			}
 
-			OutFinalValue = 1.f;
+			InOutTableData.SetFinalValue(1.0f);
 		}
 		break;
 
@@ -96,73 +108,63 @@ void FWaveTableTransform::SampleCurveBipolar(TArrayView<float> InOutTable, float
 		{
 			// Alpha is limited to between 0.0f and 1.0f and ExponentialScalar
 			// between 0 and 10 to keep ValueData "sane" and avoid float boundary.
-			for (int32 i = 0; i < InOutTable.Num(); ++i)
+			for (int32 i = 0; i < TableView.Num(); ++i)
 			{
 				ValueData[i] = ((ValueData[i] - 1.0f) * FMath::Pow(10.0f, -1.0f * Scalar * ValueData[i])) + 1.0f;
 				ValueData[i] -= 0.5f;
 			}
 
-			OutFinalValue = -1.f;
+			InOutTableData.SetFinalValue(-1.0f);
 		}
 		break;
 
 		case EWaveTableCurve::Log:
 		{
-			for (int32 i = 0; i < InOutTable.Num(); ++i)
+			for (int32 i = 0; i < TableView.Num(); ++i)
 			{
 				ValueData[i] = (Scalar * FMath::LogX(10.0f, ValueData[i])) + 0.5f;
 			}
-				
-			OutFinalValue = 1.f;
+
+			InOutTableData.SetFinalValue(1.0f);
 		}
 		break;
 
 		case EWaveTableCurve::Sin:
 		{
-			for (int32 i = 0; i < InOutTable.Num(); ++i)
+			for (int32 i = 0; i < TableView.Num(); ++i)
 			{
 				ValueData[i] = FMath::Sin(HALF_PI * ValueData[i]) - 0.5f;
 			}
 
-			OutFinalValue = 1.f;
+			InOutTableData.SetFinalValue(1.0f);
 		}
 		break;
 
 		case EWaveTableCurve::SCurve:
 		{
-			for (int32 i = 0; i < InOutTable.Num(); ++i)
+			for (int32 i = 0; i < TableView.Num(); ++i)
 			{
 				ValueData[i] = 0.5f * FMath::Sin(((PI * ValueData[i]) - HALF_PI));
 			}
 
-			OutFinalValue = 1.f;
+			InOutTableData.SetFinalValue(1.0f);
 		}
 		break;
 
 		case EWaveTableCurve::Sin_Full:
 		{
-			for (int32 i = 0; i < InOutTable.Num(); ++i)
+			for (int32 i = 0; i < TableView.Num(); ++i)
 			{
 				ValueData[i] = FMath::Sin(2.0f * PI * ValueData[i]);
 			}
 
-			OutFinalValue = 0.f;
+			InOutTableData.SetFinalValue(0.0f);
 		}
 		break;
 
 		case EWaveTableCurve::File:
 		{
-			if (WaveTable.IsEmpty())
-			{
-				Audio::ArraySetToConstantInplace(InOutTable, 0.0f);
-				OutFinalValue = 0.f;
-			}
-			else
-			{
-				constexpr auto InterpMode = WaveTable::FWaveTableSampler::EInterpolationMode::MaxValue;
-				WaveTable::FWaveTableSampler::Interpolate(WaveTable, InOutTable, InterpMode);
-				OutFinalValue = WaveTable.Last();
-			}
+			SampleWaveTableData(InOutTableData);
 		}
 		break;
 
@@ -173,25 +175,30 @@ void FWaveTableTransform::SampleCurveBipolar(TArrayView<float> InOutTable, float
 		break;
 	}
 
-	Audio::ArrayClampInPlace(InOutTable, -1.0f, 1.0f);
+	Audio::ArrayClampInPlace(TableView, -1.0f, 1.0f);
 }
 
-void FWaveTableTransform::SampleCurveUnipolar(TArrayView<float> InOutTable, float& OutFinalValue) const
+void FWaveTableTransform::SampleCurveUnipolar(FWaveTableData& InOutTableData) const
 {
+	checkf(InOutTableData.GetBitDepth() == EWaveTableBitDepth::IEEE_Float, TEXT("Curve sampling only supports IEEE 32-bit float bit depth"));
+
+	TArrayView<float> TableView;
+	InOutTableData.GetDataView(TableView);
+
 	// Clamp the input
-	Audio::ArrayClampInPlace(InOutTable, 0.0f, 1.0f);
+	Audio::ArrayClampInPlace(TableView, 0.0f, 1.0f);
 
-	float* ValueData = InOutTable.GetData();
+	float* ValueData = TableView.GetData();
 	switch (Curve)
 	{
 		case EWaveTableCurve::Custom:
 		{
-			for (int32 i = 0; i < InOutTable.Num(); ++i)
+			for (int32 i = 0; i < TableView.Num(); ++i)
 			{
 				ValueData[i] = CurveCustom.Eval(ValueData[i]);
 			}
-				
-			OutFinalValue = CurveCustom.Eval(1.f);
+
+			InOutTableData.SetFinalValue(CurveCustom.Eval(1.0f));
 			break;
 		}
 
@@ -199,30 +206,30 @@ void FWaveTableTransform::SampleCurveUnipolar(TArrayView<float> InOutTable, floa
 		{
 			if (CurveShared)
 			{
-				for (int32 i = 0; i < InOutTable.Num(); ++i)
+				for (int32 i = 0; i < TableView.Num(); ++i)
 				{
 					ValueData[i] = CurveShared->FloatCurve.Eval(ValueData[i]);
 				}
-				OutFinalValue = CurveShared->FloatCurve.Eval(1.f);
+				InOutTableData.SetFinalValue(CurveShared->FloatCurve.Eval(1.0f));
 			}
 			break;
 		}
 
 		case EWaveTableCurve::Linear_Inv:
 		{
-			for (int32 i = 0; i < InOutTable.Num(); ++i)
+			for (int32 i = 0; i < TableView.Num(); ++i)
 			{
-				ValueData[i] = 1.0f - InOutTable[i];
+				ValueData[i] = 1.0f - TableView[i];
 			}
 				
-			OutFinalValue = 0.f;
+			InOutTableData.SetFinalValue(0.0f);
 		}
 		break;
 
 		case EWaveTableCurve::Linear:
 		{
 			// Do nothing, as output values == clamped input values
-			OutFinalValue = 1.f;
+			InOutTableData.SetFinalValue(1.0f);
 		}
 		break;
 	
@@ -230,12 +237,12 @@ void FWaveTableTransform::SampleCurveUnipolar(TArrayView<float> InOutTable, floa
 		{
 			// Alpha is limited to between 0.0f and 1.0f and ExponentialScalar
 			// between 0 and 10 to keep ValueData[i]s "sane" and avoid float boundary.
-			for (int32 i = 0; i < InOutTable.Num(); ++i)
+			for (int32 i = 0; i < TableView.Num(); ++i)
 			{
 				ValueData[i] *= (FMath::Pow(10.0f, Scalar * (ValueData[i] - 1.0f)));
 			}
 
-			OutFinalValue = 1.f;
+			InOutTableData.SetFinalValue(1.0f);
 		}
 		break;
 
@@ -243,72 +250,62 @@ void FWaveTableTransform::SampleCurveUnipolar(TArrayView<float> InOutTable, floa
 		{
 			// Alpha is limited to between 0.0f and 1.0f and ExponentialScalar
 			// between 0 and 10 to keep ValueData[i]s "sane" and avoid float boundary.
-			for (int32 i = 0; i < InOutTable.Num(); ++i)
+			for (int32 i = 0; i < TableView.Num(); ++i)
 			{
 				ValueData[i] = ((ValueData[i] - 1.0f) * FMath::Pow(10.0f, -1.0f * Scalar * ValueData[i])) + 1.0f;
 			}
 
-			OutFinalValue = 0.f;
+			InOutTableData.SetFinalValue(0.0f);
 		}
 		break;
 
 		case EWaveTableCurve::Log:
 		{
-			for (int32 i = 0; i < InOutTable.Num(); ++i)
+			for (int32 i = 0; i < TableView.Num(); ++i)
 			{
 				ValueData[i] = (Scalar * FMath::LogX(10.0f, ValueData[i])) + 1.0f;
 			}
 
-			OutFinalValue = 1.f;
+			InOutTableData.SetFinalValue(1.0f);
 		}
 		break;
 
 		case EWaveTableCurve::Sin:
 		{
-			for (int32 i = 0; i < InOutTable.Num(); ++i)
+			for (int32 i = 0; i < TableView.Num(); ++i)
 			{
 				ValueData[i] = FMath::Sin(HALF_PI * ValueData[i]);
 			}
 				
-			OutFinalValue = 1.f;
+			InOutTableData.SetFinalValue(1.0f);
 		}
 		break;
 
 		case EWaveTableCurve::SCurve:
 		{
-			for (int32 i = 0; i < InOutTable.Num(); ++i)
+			for (int32 i = 0; i < TableView.Num(); ++i)
 			{
 				ValueData[i] = 0.5f * FMath::Sin(((PI * ValueData[i]) - HALF_PI)) + 0.5f;
 			}
 
-			OutFinalValue = 1.f;
+			InOutTableData.SetFinalValue(1.0f);
 		}
 		break;
 
 		case EWaveTableCurve::Sin_Full:
 		{
-			for (int32 i = 0; i < InOutTable.Num(); ++i)
+			for (int32 i = 0; i < TableView.Num(); ++i)
 			{
 				ValueData[i] = 0.5f * FMath::Sin(2.0f * PI * ValueData[i]) + 0.5f;
 			}
 
-			OutFinalValue = 0.f;
+			InOutTableData.SetFinalValue(0.0f);
 		}
 		break;
 
 		case EWaveTableCurve::File:
 		{
-			if (WaveTable.IsEmpty())
-			{
-				Audio::ArraySetToConstantInplace(InOutTable, 0.0f);
-				OutFinalValue = 0.f;
-			}
-			else
-			{
-				constexpr auto InterpMode = WaveTable::FWaveTableSampler::EInterpolationMode::MaxValue;
-				WaveTable::FWaveTableSampler::Interpolate(WaveTable, InOutTable, InterpMode);
-				OutFinalValue = WaveTable.Last();
-			}
+			SampleWaveTableData(InOutTableData);
 		}
 		break;
 
@@ -319,13 +316,34 @@ void FWaveTableTransform::SampleCurveUnipolar(TArrayView<float> InOutTable, floa
 		break;
 	}
 
-	Audio::ArrayClampInPlace(InOutTable, 0.0f, 1.0f);
+	Audio::ArrayClampInPlace(TableView, 0.0f, 1.0f);
+}
+
+void FWaveTableTransform::SampleWaveTableData(FWaveTableData& InOutTableData) const
+{
+	checkf(InOutTableData.GetBitDepth() == EWaveTableBitDepth::IEEE_Float, TEXT("WaveTable resampling only supports IEEE 32-bit float bit depth"));
+
+	TArrayView<float> TableView;
+	InOutTableData.GetDataView(TableView);
+
+	if (TableData.IsEmpty())
+	{
+		Audio::ArraySetToConstantInplace(TableView, 0.0f);
+		InOutTableData.SetFinalValue(0.0f);
+	}
+	else
+	{
+		constexpr auto InterpMode = WaveTable::FWaveTableSampler::EInterpolationMode::MaxValue;
+		WaveTable::FWaveTableSampler::Interpolate(TableData, TableView, InterpMode);
+		InOutTableData.SetFinalValue(TableData.GetFinalValue());
+	}
 }
 
 void FWaveTableTransform::CacheCurve()
 {
 	if (Curve == EWaveTableCurve::Shared && CurveShared)
 	{
+		Curve = EWaveTableCurve::Custom;
 		CurveCustom = CurveShared->FloatCurve;
 	}
 
@@ -337,45 +355,80 @@ void FWaveTableTransform::CacheCurve()
 
 float FWaveTableTransform::GetFinalValue() const
 {
-	if (FMath::IsNearlyEqual(FinalValue, WaveTable::InvalidWaveTableValue))
-	{
-		if (WaveTable.IsEmpty())
-		{
-			return 0.f;
-		}
-
-		return WaveTable.Last();
-	}
-
-	return FinalValue;
+	return TableData.GetFinalValue();
 }
 
 void FWaveTableTransform::SetFinalValue(float InValue)
 {
-	FinalValue = InValue;
+}
+
+const FWaveTableData& FWaveTableTransform::GetTableData() const
+{
+	return TableData;
 }
 
 #if WITH_EDITOR
 void FWaveTableTransform::CopyToWaveTable(TArrayView<float> InOutTable, float& OutFinalValue, bool bInBipolar) const
 {
-	const float Delta = 1.0f / InOutTable.Num();
-	
-	for (int32 i = 0; i < InOutTable.Num(); ++i)
+	constexpr bool bIsLoop = false;
+
+	FWaveTableData Data;
+	Data.SetData(InOutTable, bIsLoop);
+	Data.SetFinalValue(OutFinalValue);
+
+	CopyToWaveTable(Data, bInBipolar);
+
+	ensureAlwaysMsgf(Data.GetDataView(InOutTable), TEXT("BitDepth format mismatch"));
+	OutFinalValue = Data.GetFinalValue();
+}
+
+void FWaveTableTransform::CopyToWaveTable(FWaveTableData& InOutTableData, bool bInBipolar) const
+{
+	FWaveTableData SampleData(EWaveTableBitDepth::IEEE_Float);
+	SampleData.Zero(InOutTableData.GetNumSamples());
+
+	TArrayView<float> SampleDataTableView;
+	SampleData.GetDataView(SampleDataTableView);
+	const float Delta = 1.0f / InOutTableData.GetNumSamples();
+	for (int32 Index = 0; Index < SampleDataTableView.Num(); ++Index)
 	{
-		InOutTable[i] = i * Delta;
+		SampleDataTableView[Index] = Index * Delta;
 	}
 
-	bInBipolar ? SampleCurveBipolar(InOutTable, OutFinalValue) : SampleCurveUnipolar(InOutTable, OutFinalValue);
+	bInBipolar ? SampleCurveBipolar(SampleData) : SampleCurveUnipolar(SampleData);
+
+	switch (InOutTableData.GetBitDepth())
+	{
+		case EWaveTableBitDepth::IEEE_Float:
+		{
+			InOutTableData = MoveTemp(SampleData);
+		}
+		break;
+
+		case EWaveTableBitDepth::PCM_16:
+		{
+			TArrayView<int16> TableView;
+			InOutTableData.GetDataView(TableView);
+			Audio::ArrayFloatToPcm16(SampleDataTableView, TableView);
+		}
+		break;
+
+		default:
+		{
+			static_assert(static_cast<int32>(EWaveTableBitDepth::COUNT) == 2, "Possible missing switch case coverage for 'EWaveTableBitDepth'");
+			checkNoEntry();
+		}
+	}
 }
 
-void FWaveTableTransform::CreateWaveTable(TArray<float>& InOutTable, float& OutFinalValue, bool bInBipolar) const
+void FWaveTableTransform::CreateWaveTable(FWaveTableData& InOutTableData, bool bInBipolar) const
 {
 	switch (Curve)
 	{
 		case EWaveTableCurve::File:
 		{
 			WaveTable::FImporter Importer(WaveTableSettings, bInBipolar);
-			Importer.Process(InOutTable);
+			Importer.Process(InOutTableData);
 		}
 		break;
 
@@ -391,9 +444,25 @@ void FWaveTableTransform::CreateWaveTable(TArray<float>& InOutTable, float& OutF
 		case EWaveTableCurve::Sin_Full:
 		default:
 		{
-			CopyToWaveTable(InOutTable, OutFinalValue, bInBipolar);
+			CopyToWaveTable(InOutTableData, bInBipolar);
 		}
 		break;
 	}
 }
+
+void FWaveTableTransform::CreateWaveTable(TArray<float>& InOutTable, float& OutFinalValue, bool bInBipolar) const
+{
+	constexpr bool bIsLoop = false;
+
+	FWaveTableData TempTableData;
+	TempTableData.SetData(InOutTable, bIsLoop);
+	TempTableData.SetFinalValue(OutFinalValue);
+
+	CreateWaveTable(TempTableData, bInBipolar);
+
+	TArrayView<float> OutTableView;
+	TempTableData.GetDataView(OutTableView);
+	OutFinalValue = TableData.GetFinalValue();
+	InOutTable = MoveTemp(OutTableView);
+}
 #endif // WITH_EDITOR
diff --git a/Engine/Plugins/Runtime/WaveTable/Source/WaveTable/Public/WaveTable.h b/Engine/Plugins/Runtime/WaveTable/Source/WaveTable/Public/WaveTable.h
--- a/Engine/Plugins/Runtime/WaveTable/Source/WaveTable/Public/WaveTable.h
+++ b/Engine/Plugins/Runtime/WaveTable/Source/WaveTable/Public/WaveTable.h
@@ -2,102 +2,134 @@
 #pragma once
 
 #include "Containers/ArrayView.h"
+#include "WaveTableSampler.h"
+
+#include "WaveTable.generated.h"
 
 
 namespace WaveTable
 {
+	// An invalid WaveTable value, used to detect uninitialized final values.
+	constexpr float InvalidWaveTableValue = TNumericLimits<float>::Min();
+
+	// Read-only view to a wave table struct
 	struct WAVETABLE_API FWaveTableView
 	{
 		TArrayView<const float> SampleView;
-		float FinalValue;
+		const float FinalValue = 0.0f;
 
-		explicit FWaveTableView(const TArray<float>& InSamples, const float InFinalValue)
-			: SampleView(InSamples), FinalValue(InFinalValue) {}
-		
-		explicit FWaveTableView(const TArrayView<const float>& InSamples, const float InFinalValue)
-			: SampleView(InSamples), FinalValue(InFinalValue) {}
+		explicit FWaveTableView(const TArray<float>& InSamples, const float InFinalValue);
+		explicit FWaveTableView(const TArrayView<const float>& InSamples, const float InFinalValue);
 
 		// returns the length of the SampleView array, not including the additional FinalValue
-		int32 Num() const
-		{
-			return SampleView.Num();
-		}
-
-		bool IsEmpty() const
-		{
-			return SampleView.IsEmpty();
-		}
+		int32 Num() const;
+		bool IsEmpty() const;
 	};
-	
+
+	// WaveTable data used at runtime by oscillators, envelopes,
+	//  etc.  Always as 32-bit float bit depth for processing optimization,
+	// at the expense of memory.
 	struct WAVETABLE_API FWaveTable
 	{
 	private:
 		TArray<float> Samples;
-		float FinalValue;
+		float FinalValue = 0.0f;
 
 	public:
 		FWaveTable() = default;
-
-		FWaveTable(const TArray<float>& InSamples, const float InFinalValue=0.f)
-			: Samples(InSamples)
-			, FinalValue(InFinalValue)
-		{
-		}
-
-		FWaveTable(TArray<float>&& InSamples, const float InFinalValue)
-			: Samples(MoveTemp(InSamples))
-			, FinalValue(InFinalValue)
-		{
-		}
-
-		float* GetData()
-		{
-			return Samples.GetData();
-		}
-
-		const float* GetData() const
-		{
-			return Samples.GetData();
-		}
-
-		FWaveTableView GetView() const
-		{
-			return FWaveTableView(Samples, FinalValue);
-		}
-
-		TArrayView<float> GetSamples()
-		{
-			return Samples;
-		}
-
-		int32 Num() const
-		{
-			return Samples.Num();
-		}
-
-		void Set(TArray<float>&& InSamples)
-		{
-			Samples = MoveTemp(InSamples);
-		}
-
-		void Set(TArrayView<const float> InSamples)
-		{
-			Samples = InSamples;
-		}
-
-		void SetFinalValue(const float InValue)
-		{
-			FinalValue = InValue;
-		}
-
-		void SetNum(int32 InSize)
-		{
-			Samples.SetNum(InSize);
-		}
-
-		void Zero()
-		{
-			FMemory::Memset(Samples.GetData(), 0, sizeof(float) * Samples.Num());
-		}
+		FWaveTable(const TArray<float>& InSamples, const float InFinalValue = 0.0f);
+		FWaveTable(TArray<float>&& InSamples, const float InFinalValue);
+		FWaveTable(const ::FWaveTableData& InData);
+
+		float* GetData();
+		const float* GetData() const;
+		FWaveTableView GetView() const;
+		TArrayView<float> GetSamples();
+		int32 Num() const;
+		void Set(TArray<float>&& InSamples);
+		void Set(TArrayView<const float> InSamples);
+		void SetFinalValue(const float InValue);
+		void SetNum(int32 InSize);
+		void Zero();
 	};
 } // namespace WaveTable
+
+
+UENUM(BlueprintType)
+enum class EWaveTableBitDepth : uint8
+{
+	// Lower resolution and marginal performance cost with
+	// conversion overhead (engine operates on 32-bit)
+	// with the advantage of half the size in memory.
+	PCM_16,
+
+	// Higher precision and faster operative performance
+	// (engine operates at 32-bit) at the cost of twice the
+	// memory of 16-bit.
+	IEEE_Float,
+
+	COUNT UMETA(Hidden)
+};
+
+
+// Serialized WaveTable data, that supports multiple bit depth formats.
+USTRUCT(BlueprintType)
+struct WAVETABLE_API FWaveTableData
+{
+	GENERATED_BODY();
+
+	FWaveTableData() = default;
+	FWaveTableData(EWaveTableBitDepth InBitDepth);
+	FWaveTableData(TArrayView<const float> InSamples, float InFinalValue);
+	FWaveTableData(TArrayView<const int16> InSamples, float InFinalValue);
+
+#if WITH_EDITOR
+	static FName GetBitDepthPropertyName();
+#endif // WITH_EDITOR
+
+private:
+	UPROPERTY(EditAnywhere, Category = Options)
+	EWaveTableBitDepth BitDepth = EWaveTableBitDepth::PCM_16;
+
+	UPROPERTY()
+	TArray<uint8> Data;
+
+	UPROPERTY()
+	float FinalValue = ::WaveTable::InvalidWaveTableValue;
+
+	FORCEINLINE int32 GetSampleSize() const;
+
+public:
+	// Adds this WaveTableData's internal contents to the given float buffer view,
+	// initially performing bit depth conversion to values to IEEE_FLOAT if required.
+	// Provided buffer's size must match this WaveTableData's samples count.
+	void ArrayMixIn(TArrayView<float> OutputWaveSamples, float Gain = 1.0f) const;
+
+	void Empty();
+
+	EWaveTableBitDepth GetBitDepth() const;
+	const TArray<uint8>& GetRawData() const;
+	bool GetDataView(TArrayView<const float>& OutData) const;
+	bool GetDataView(TArrayView<float>& OutData);
+	bool GetDataView(TArrayView<const int16>& OutData) const;
+	bool GetDataView(TArrayView<int16>& OutData);
+	int32 GetNumSamples() const;
+	float GetFinalValue() const;
+
+	bool IsEmpty() const;
+
+	bool Resample(int32 InCurrentSampleRate, int32 InNewSampleRate, ::WaveTable::FWaveTableSampler::EInterpolationMode InInterpMode = ::WaveTable::FWaveTableSampler::EInterpolationMode::Cubic);
+	void Reset(int32 NumSamples);
+
+	// Converts bit depth and sets internal data to the requested bit depth.
+	// Returns true if conversion took place, false if not (i.e. data was 
+	// already sampled at the given bit depth).
+	bool SetBitDepth(EWaveTableBitDepth InBitDepth);
+	void SetData(TArrayView<const int16> InData, bool bIsLoop);
+	void SetData(TArrayView<const float> InData, bool bIsLoop);
+	void SetRawData(EWaveTableBitDepth InBitDepth, TArrayView<const uint8> InData);
+	void SetRawData(EWaveTableBitDepth InBitDepth, TArray<uint8>&& InData);
+	void SetFinalValue(float InFinalValue);
+
+	void Zero(int32 InNumSamples = INDEX_NONE);
+};
diff --git a/Engine/Plugins/Runtime/WaveTable/Source/WaveTable/Public/WaveTableBank.h b/Engine/Plugins/Runtime/WaveTable/Source/WaveTable/Public/WaveTableBank.h
--- a/Engine/Plugins/Runtime/WaveTable/Source/WaveTable/Public/WaveTableBank.h
+++ b/Engine/Plugins/Runtime/WaveTable/Source/WaveTable/Public/WaveTableBank.h
@@ -7,7 +7,12 @@
 
 #include "WaveTableBank.generated.h"
 
+
+// Forward Declarations
 struct FPropertyChangedEvent;
+class FWaveTableBankAssetProxy;
+
+using FWaveTableBankAssetProxyPtr = TSharedPtr<FWaveTableBankAssetProxy, ESPMode::ThreadSafe>;
 
 
 USTRUCT()
@@ -25,10 +30,18 @@ class WAVETABLE_API UWaveTableBank : public UObject, public IAudioProxyDataFacto
 	GENERATED_BODY()
 
 public:
-	// Number of samples cached for each curve in the given bank.
+	// Sampling mode used for the bank.
 	UPROPERTY(EditAnywhere, Category = Options)
+	EWaveTableSamplingMode SampleMode = EWaveTableSamplingMode::FixedResolution;
+
+	// Number of samples cached for each entry in the given bank.
+	UPROPERTY(EditAnywhere, Category = Options, meta = (EditCondition = "SampleMode == EWaveTableSamplingMode::FixedResolution", EditConditionHides))
 	EWaveTableResolution Resolution = EWaveTableResolution::Res_256;
 
+	// Number of samples cached for each entry in the given bank.
+	UPROPERTY(EditAnywhere, Category = Options, meta = (EditCondition = "SampleMode == EWaveTableSamplingMode::FixedSampleRate", ClampMin = "1", ClampMax = "48000", EditConditionHides))
+	int32 SampleRate = 48000;
+
 	// Determines if output from curve/wavetable are to be clamped between 
 	// [-1.0f, 1.0f] (i.e. for waveform generation, oscillation, etc.)
 	// or [0.0f, 1.0f] (i.e. for enveloping) when sampling curve/wavetable
@@ -41,17 +54,25 @@ public:
 	float WaveTableSizeMB = 0.0f;
 
 	// Length of all WaveTable samples in bank in seconds (at 48kHz)
-	UPROPERTY(VisibleAnywhere, Category = Options, meta = (DisplayName = "WaveTable Length (sec)"))
+	UE_DEPRECATED(5.3, "Samples now each have own length, as they no longer are required being the same length if using shared 'SampleRate' mode")
 	float WaveTableLengthSec = 0.0f;
 #endif // WITH_EDITORONLY_DATA
 
+
 	/** Tables within the given bank */
+	UE_DEPRECATED(5.3, "Direct access of 'Entries' will become protected member in future release and not externally modifiable in runtime builds. If in editor, use GetEntries(). "
+		"To reduce memory consumption, entries are now readonly and cleared in runtime builds when proxy is generated.")
 	UPROPERTY(EditAnywhere, Category = Options)
 	TArray<FWaveTableBankEntry> Entries;
 
+public:
 	/* IAudioProxyDataFactory Implementation */
 	virtual TSharedPtr<Audio::IProxyData> CreateProxyData(const Audio::FProxyDataInitParams& InitParams) override;
 
+	virtual void Serialize(FArchive& Ar) override;
+
+	TArray<FWaveTableBankEntry>& GetEntries();
+
 #if WITH_EDITOR
 	void RefreshWaveTables();
 
@@ -59,6 +80,12 @@ public:
 	virtual void PostEditChangeProperty(FPropertyChangedEvent& InPropertyChangedEvent) override;
 	virtual void PostEditChangeChainProperty(FPropertyChangedChainEvent& InPropertyChangedEvent) override;
 #endif // WITH_EDITOR
+
+private:
+	void CopyToProxyData();
+	void MoveToProxyData();
+
+	FWaveTableBankAssetProxyPtr ProxyData;
 };
 
 class WAVETABLE_API FWaveTableBankAssetProxy : public Audio::TProxyData<FWaveTableBankAssetProxy>, public TSharedFromThis<FWaveTableBankAssetProxy, ESPMode::ThreadSafe>
@@ -66,24 +93,39 @@ class WAVETABLE_API FWaveTableBankAssetProxy : public Audio::TProxyData<FWaveTab
 public:
 	IMPL_AUDIOPROXY_CLASS(FWaveTableBankAssetProxy);
 
-	FWaveTableBankAssetProxy(const UWaveTableBank& InWaveTableBank)
+	FWaveTableBankAssetProxy() = default;
+	FWaveTableBankAssetProxy(const FWaveTableBankAssetProxy& InAssetProxy);
+	FWaveTableBankAssetProxy(uint32 InObjectId, EWaveTableSamplingMode InSamplingMode, int32 InSampleRate, const TArray<FWaveTableBankEntry>& InBankEntries);
+	FWaveTableBankAssetProxy(uint32 InObjectId, EWaveTableSamplingMode InSamplingMode, int32 InSampleRate, TArray<FWaveTableBankEntry>&& InBankEntries);
+
+	UE_DEPRECATED(5.3, "Proxy generation & respective data translation is now entirely handled by UWaveTableBank calling the constructor variants above.")
+	FWaveTableBankAssetProxy(const UWaveTableBank& InWaveTableBank);
+
+	virtual ~FWaveTableBankAssetProxy() = default;
+
+	virtual EWaveTableSamplingMode GetSampleMode() const
+	{
+		return SampleMode;
+	}
+
+	virtual int32 GetSampleRate() const
 	{
-		Algo::Transform(InWaveTableBank.Entries, WaveTables, [](const FWaveTableBankEntry& Entry)
-		{
-			return WaveTable::FWaveTable(Entry.Transform.WaveTable, Entry.Transform.GetFinalValue());
-		});
+		return SampleRate;
 	}
 
-	virtual const TArray<WaveTable::FWaveTable>& GetWaveTables() const
+	virtual const TArray<FWaveTableData>& GetWaveTableData() const
 	{
-		return WaveTables;
+		return WaveTableData;
+	}
+
+	uint32 GetObjectId() const
+	{
+		return ObjectId;
 	}
 
 protected:
-	TArray<WaveTable::FWaveTable> WaveTables;
+	uint32 ObjectId = INDEX_NONE;
+	int32 SampleRate = 48000;
+	EWaveTableSamplingMode SampleMode = EWaveTableSamplingMode::FixedResolution;
+	TArray<FWaveTableData> WaveTableData;
 };
-using FWaveTableBankAssetProxyPtr = TSharedPtr<FWaveTableBankAssetProxy, ESPMode::ThreadSafe>;
-
-#if UE_ENABLE_INCLUDE_ORDER_DEPRECATED_IN_5_2
-#include "UObject/ObjectSaveContext.h"
-#endif
diff --git a/Engine/Plugins/Runtime/WaveTable/Source/WaveTable/Public/WaveTableImporter.h b/Engine/Plugins/Runtime/WaveTable/Source/WaveTable/Public/WaveTableImporter.h
--- a/Engine/Plugins/Runtime/WaveTable/Source/WaveTable/Public/WaveTableImporter.h
+++ b/Engine/Plugins/Runtime/WaveTable/Source/WaveTable/Public/WaveTableImporter.h
@@ -1,9 +1,12 @@
 // Copyright Epic Games, Inc. All Rights Reserved.
 #pragma once
 
+#include "WaveTable.h"
 #include "WaveTableSampler.h"
+#include "WaveTableSettings.h"
 
-struct FWaveTableSettings;
+// Forward Declarations
+struct FWaveTableData;
 
 
 namespace WaveTable
@@ -13,9 +16,14 @@ namespace WaveTable
 	public:
 		FImporter(const FWaveTableSettings& InSettings, bool bInBipolar);
 
+		UE_DEPRECATED(5.3, "Importer now supports multiple bit depths and sample rate. Use version of function that takes in an FWaveTableData struct.")
 		void Process(TArray<float>& OutWaveTable);
 
+		void Process(FWaveTableData& OutData);
+
 	private:
+		void ProcessInternal(const TArrayView<const float> InEditSourceView, TArray<float>& OutWaveTable);
+
 		const FWaveTableSettings& Settings;
 
 		bool bBipolar = false;
@@ -23,7 +31,3 @@ namespace WaveTable
 		FWaveTableSampler Sampler;
 	};
 } // namespace WaveTable
-
-#if UE_ENABLE_INCLUDE_ORDER_DEPRECATED_IN_5_2
-#include "WaveTableSettings.h"
-#endif
diff --git a/Engine/Plugins/Runtime/WaveTable/Source/WaveTable/Public/WaveTableSampler.h b/Engine/Plugins/Runtime/WaveTable/Source/WaveTable/Public/WaveTableSampler.h
--- a/Engine/Plugins/Runtime/WaveTable/Source/WaveTable/Public/WaveTableSampler.h
+++ b/Engine/Plugins/Runtime/WaveTable/Source/WaveTable/Public/WaveTableSampler.h
@@ -4,6 +4,8 @@
 #include "Containers/ArrayView.h"
 
 
+struct FWaveTableData;
+
 namespace WaveTable
 {
 	struct FWaveTableView;
@@ -58,34 +60,65 @@ namespace WaveTable
 			float Phase = 0.0f;
 
 			EInterpolationMode InterpolationMode = EInterpolationMode::Linear;
+
+			bool bOneShot = false;
 		};
 
 		FWaveTableSampler();
 		FWaveTableSampler(FSettings&& InSettings);
 
-		// Interpolates and converts values in the given table for each provided index in the index-to-samples TArrayView (an array of sub-sample, floating point, indices)
+		// Interpolates and converts values in the given table for each provided index in the index-to-samples TArrayView (an array of sub-sample, floating point indices)
+		static void Interpolate(const FWaveTableData& InTableData, TArrayView<float> InOutIndexToSamplesView, EInterpolationMode InterpMode = EInterpolationMode::Linear);
+		static void Interpolate(TArrayView<const int16> InTableView, TArrayView<float> InOutIndexToSamplesView, EInterpolationMode InterpMode = EInterpolationMode::Linear);
 		static void Interpolate(TArrayView<const float> InTableView, TArrayView<float> InOutIndexToSamplesView, EInterpolationMode InterpMode = EInterpolationMode::Linear);
 
+		// Retrieves the sample value at the currently set phase, returning the floating point index which the current phase corresponds to.
 		float Process(const FWaveTableView& InTableView, float& OutSample, ESingleSampleMode InMode = ESingleSampleMode::Zero);
+
+		// Resamples entire table length into given view
 		float Process(const FWaveTableView& InTableView, TArrayView<float> OutSamplesView);
+		float Process(const FWaveTableData& InTableData, TArrayView<float> OutSamplesView);
 		float Process(const FWaveTableView& InTableView, TArrayView<const float> InFreqModulator, TArrayView<const float> InPhaseModulator, TArrayView<const float> InSyncTriggers, TArrayView<float> OutSamplesView);
+		float Process(const FWaveTableData& InTableData, TArrayView<const float> InFreqModulator, TArrayView<const float> InPhaseModulator, TArrayView<const float> InSyncTriggers, TArrayView<float> OutSamplesView);
 
+		// Resets the last index operated on, restarting the interpolation process for subsequent calls.
 		void Reset();
 
 		const FSettings& GetSettings() const;
 
+		// If set to one-shot, returns the first index where a reverse or stall occurred in the provided
+		// interpolated index, ignoring samples where a "sync" request took place. If set to looping,
+		// always INDEX_NONE as looping does not support a "finished" state.
+		int32 GetIndexFinished() const;
+
 		void SetInterpolationMode(EInterpolationMode InMode);
 		void SetFreq(float InFreq);
+		void SetOneShot(bool bInLooping);
 		void SetPhase(float InPhase);
 
 	private:
-		void ComputeIndexFrequency(TArrayView<const float> InTableView, TArrayView<const float> InFreqModulator, TArrayView<const float> InSyncTriggers, TArrayView<float> OutIndicesView);
-		void ComputeIndexPhase(TArrayView<const float> InTableView, TArrayView<const float> InPhaseModulator, TArrayView<float> OutIndicesView);
+		void ComputeIndexFrequency(int32 NumInputSamples, TArrayView<const float> InFreqModulator, TArrayView<const float> InSyncTriggers, TArrayView<float> OutIndicesView);
+		void ComputeIndexPhase(int32 NumInputSamples, TArrayView<const float> InPhaseModulator, TArrayView<float> OutIndicesView);
+
+		// If set to one-shot, sets IndexFinished to valid index and resets remaining indices provided past the initial stop to the first index.
+		// Returns last positively increasing index if finished. If not finished, returns the last index value stored in the view.
+		int32 ComputeIndexFinished(TArrayView<const float> InSyncTriggers, TArrayView<float> OutIndicesView);
 
 		float LastIndex = 0.0f;
 
+		struct FOneShotData
+		{
+			// Cached last floating index from prior process, used to compute stopping index if forward progress terminates on buffer boundary.
+			float LastOutputIndex = -1.0f;
+
+			// Contains index where sampler finished as a result of a sample ending or phase reversing within the last process call.
+			// (If < 0, no finish/phase reverse was detected and sampler is considered to be continuing to make forward progress)
+			int32 IndexFinished = INDEX_NONE;
+		} OneShotData;
+
 		TArray<float> PhaseModScratch;
 		FSettings Settings;
+
 	};
 } // namespace WaveTable
 
diff --git a/Engine/Plugins/Runtime/WaveTable/Source/WaveTable/Public/WaveTableSettings.h b/Engine/Plugins/Runtime/WaveTable/Source/WaveTable/Public/WaveTableSettings.h
--- a/Engine/Plugins/Runtime/WaveTable/Source/WaveTable/Public/WaveTableSettings.h
+++ b/Engine/Plugins/Runtime/WaveTable/Source/WaveTable/Public/WaveTableSettings.h
@@ -1,10 +1,16 @@
 // Copyright Epic Games, Inc. All Rights Reserved.
 #pragma once
 
+#include "WaveTable.h"
 
 #include "WaveTableSettings.generated.h"
 
 
+// Forward Declarations
+class UWaveTableBank;
+struct FWaveTableBankEntry;
+
+
 namespace WaveTable
 {
 	// Converts a ratio [0.0, 1.0] to an index within the given number of array elements.
@@ -45,6 +51,28 @@ enum class EWaveTableResolution : uint8
 };
 
 
+// SamplingMode of a given bank or collection of WaveTables
+UENUM()
+enum class EWaveTableSamplingMode : uint8
+{
+	// Enforces fixed sample rate for all members in the bank/collection,
+	// enabling them to be of unique duration/number of samples. Good for
+	// use cases when entries are being treated as separate, discrete but
+	// related audio to be played back at a shared speed (ex. traditional
+	// "samplers" or granulation).
+	FixedSampleRate,
+
+	// Enforces resolution (i.e. size of all tables), uniformly resampling
+	// all tables in the collection to be the same length/number of samples
+	// (if not already).  Supports use cases where systems are mixing/
+	// interpolating or spatializing entries in lockstep (ex. oscillating
+	// or enveloping).
+	FixedResolution,
+
+	COUNT UMETA(Hidden)
+};
+
+
 USTRUCT()
 struct WAVETABLE_API FWaveTableSettings
 {
@@ -58,6 +86,14 @@ struct WAVETABLE_API FWaveTableSettings
 	UPROPERTY(EditAnywhere, Category = Options, meta = (ClampMin = "0.0"))
 	int32 ChannelIndex = 0;
 
+	// Source data last imported from the source file
+	UPROPERTY(EditAnywhere, Category = Options)
+	FWaveTableData SourceData;
+
+	// Source sample rate from last imported source file
+	UPROPERTY(VisibleAnywhere, Category = Options)
+	int32 SourceSampleRate = 48000;
+
 	// Percent to phase shift of table
 	UPROPERTY(EditAnywhere, Category = Options, meta = (ClampMin = "0.0", ClampMax = "1.0"))
 	float Phase = 0.0f;
@@ -79,7 +115,7 @@ struct WAVETABLE_API FWaveTableSettings
 	float FadeOut = 0.0f;
 
 	// Whether or not to normalize the WaveTable.
-	UPROPERTY(EditAnywhere, Category = Options, meta = (DisplayAfter = "FilePath"))
+	UPROPERTY(EditAnywhere, Category = Options, meta = (DisplayAfter = "SourceData"))
 	bool bNormalize = true;
 
 	// Whether or not to remove offset from original file
@@ -87,14 +123,31 @@ struct WAVETABLE_API FWaveTableSettings
 	UPROPERTY(EditAnywhere, Category = Options, meta = (DisplayAfter = "bNormalize"))
 	bool bRemoveOffset = false;
 
-	// SourcePCM Data
-	UPROPERTY()
-	TArray<float> SourcePCMData;
+	// Whether or not an edit is contained to omit or modify source file's PCM data
+	bool ContainsEdit() const
+	{
+		return Phase > 0.0f ||
+			Top > 0.0f || Tail > 0.0f ||
+			FadeIn > 0.0f || FadeOut > 0.0f ||
+			bNormalize || bRemoveOffset;
+	}
+
+#if WITH_EDITOR
+	void VersionPCMData();
+#endif // WITH_EDITOR
 
-	// Returns a TArrayView of Source PCM values that are on when
-	// importing with the given settings (restricted effectively
-	// due to top/tail settings).
+	// Returns the respective offset & number of samples that
+	// characterize the source data to apply light edits to.
+	void GetEditSourceBounds(int32& OutSourceTopOffset, int32& OutSourceNumSamples) const;
+
+	UE_DEPRECATED(5.3, "Source now supports multiple bit depths. Use 'GetEditSourceBounds' & apply to SourceData WaveTableData TArrayView in proper format using 'GetDataView'.")
 	TArrayView<const float> GetEditSourceView() const;
+
+private:
+#if WITH_EDITORONLY_DATA
+	UPROPERTY(meta = (Deprecated = "5.3", DeprecationMessage = "Migrated to FWaveTableData member 'SourceData' to support multiple bitdepths"))
+	TArray<float> SourcePCMData;
+#endif // WITH_EDITORONLY_DATA
 };
 
 UENUM(BlueprintType)
@@ -124,7 +177,6 @@ enum class EWaveTableCurve : uint8
 
 namespace WaveTable
 {
-	// Returns WaveTable size given WaveTableResolution and optional Curve or PCM size. If resolution is none, uses default value associated
-	// with provided curve. If set to file, will use either the largest static resolution (Res_Max, or 4096) or InMaxPCMSize, whichever is larger.
+	UE_DEPRECATED(5.3, "Removed from public API as banks now support different bit depths and sampling mode (i.e. sample count can no longer be conflated with generic size metric)")
 	WAVETABLE_API int32 GetWaveTableSize(EWaveTableResolution InWaveTableResolution, EWaveTableCurve InCurve, int32 InMaxPCMSize = INDEX_NONE);
 } // namespace WaveTable
diff --git a/Engine/Plugins/Runtime/WaveTable/Source/WaveTable/Public/WaveTableTransform.h b/Engine/Plugins/Runtime/WaveTable/Source/WaveTable/Public/WaveTableTransform.h
--- a/Engine/Plugins/Runtime/WaveTable/Source/WaveTable/Public/WaveTableTransform.h
+++ b/Engine/Plugins/Runtime/WaveTable/Source/WaveTable/Public/WaveTableTransform.h
@@ -1,18 +1,15 @@
 // Copyright Epic Games, Inc. All Rights Reserved.
 #pragma once
 
+#include "Containers/ArrayView.h"
 #include "Curves/RichCurve.h"
+#include "WaveTable.h"
 #include "WaveTableSettings.h"
 
 #include "WaveTableTransform.generated.h"
 
 class UCurveFloat;
 
-// use a value outside the [-1,1] range to detect uninitialized values
-namespace WaveTable
-{
-	constexpr float InvalidWaveTableValue = TNumericLimits<float>::Min();
-}
 
 USTRUCT(BlueprintType)
 struct WAVETABLE_API FWaveTableTransform
@@ -29,7 +26,7 @@ struct WAVETABLE_API FWaveTableTransform
 	 *  Logarithmic: y = b * log(x) + 1
 	 */
 	UPROPERTY(EditAnywhere, Category = Input, BlueprintReadWrite, meta = (DisplayName = "Exponential Scalar", ClampMin = "0.1", ClampMax = "10.0", UIMin = "0.1", UIMax = "10.0"))
-	float Scalar = 2.5;
+	float Scalar = 2.5f;
 
 	/** Custom curve to apply if output curve type is set to 'Custom.' */
 	UPROPERTY()
@@ -39,9 +36,17 @@ struct WAVETABLE_API FWaveTableTransform
 	UPROPERTY(EditAnywhere, meta = (DisplayName = "Asset"), Category = Curve, BlueprintReadWrite)
 	TObjectPtr<UCurveFloat> CurveShared = nullptr;
 
+protected:
 	UPROPERTY()
-	TArray<float> WaveTable;
+	FWaveTableData TableData;
+
+	// Duration of curve or file.  Only valid if parent SampleRate is set and SamplingMode is set to 'FixedSampleRate'
+	// (If set to 'FixedResolution', duration is variable depending on the resolution (number of samples in the table data)
+	// and device's sample rate.
+	UPROPERTY(EditAnywhere, meta = (DisplayName = "Duration (Sec)", Category = Input, ClampMin = "0.0", UIMin = "0.0"))
+	float Duration = 1.0f;
 
+public:
 #if WITH_EDITORONLY_DATA
 	UPROPERTY(EditAnywhere, Category = WaveTable)
 	FWaveTableSettings WaveTableSettings;
@@ -50,16 +55,47 @@ struct WAVETABLE_API FWaveTableTransform
 	// Applies transform to provided index value, setting it to the corresponding value.
 	void Apply(float& InOutValue, bool bInBipolar = false) const;
 
+	float GetDuration() const { return Duration; }
+
 #if WITH_EDITOR
-	// Copies a single, period [0.0, 1.0] to a WaveTable using this transform's curve. The
+	UE_DEPRECATED(5.3, "Use FWaveTableData version instead")
+	void CopyToWaveTable(TArrayView<float> InOutTable, float& OutFinalValue, bool bInBipolar) const;
+
+	UE_DEPRECATED(5.3, "Use FWaveTableData version instead")
+	void CreateWaveTable(TArray<float>& InOutTable, float& OutFinalValue, bool bInBipolar) const;
+
+	// Copies a single, period [0.0, 1.0] to provided WaveTableData using this transform's curve. The
 	// size of the resulting WaveTable is the provided ArrayView's size. If curve is itself
 	// set to 'File' and internally cached as WaveTable, will re-sample using max value
-	// interpolation. Does *not* apply destructive edits for using WaveTableSettings PCM data.
-	void CopyToWaveTable(TArrayView<float> InOutTable, float& OutFinalValue, bool bInBipolar) const;
+	// interpolation. Does *not* apply destructive edits from WaveTableSetting's source table data.
+	void CopyToWaveTable(FWaveTableData& InOutTableData, bool bInBipolar) const;
 
-	// Caches WaveTable for curve using the provided array length as resolution.
+	// Creates WaveTableData from transform respective Curve settings.
 	// Applies destructive edits as provided by WaveTableSettings if curve set to 'File'.
-	void CreateWaveTable(TArray<float>& InOutTable, float& OutFinalValue, bool bInBipolar) const;
+	void CreateWaveTable(FWaveTableData& InOutTableData, bool bInBipolar) const;
+
+	static FName GetDurationPropertyName()
+	{
+		return GET_MEMBER_NAME_CHECKED(FWaveTableTransform, Duration);
+	}
+
+	// Versions deprecated "WaveTable" property data to PCMData, to support serialization of different bit depths
+	void VersionTableData()
+	{
+		WaveTableSettings.VersionPCMData();
+
+		PRAGMA_DISABLE_DEPRECATION_WARNINGS
+		if (!WaveTable.IsEmpty())
+		{
+			constexpr bool bIsLoop = false;
+			TableData.SetData(WaveTable, bIsLoop);
+			TableData.SetFinalValue(FinalValue);
+
+			FinalValue = WaveTable::InvalidWaveTableValue;
+			WaveTable.Empty();
+		}
+		PRAGMA_ENABLE_DEPRECATION_WARNINGS
+	}
 #endif // WITH_EDITOR
 
 	/** Caches curve data.  Should a shared curve be selected, curve points are copied locally
@@ -69,19 +105,32 @@ struct WAVETABLE_API FWaveTableTransform
 	  */
 	void CacheCurve();
 
-	// Returns the final value of the WaveTable
+	const FWaveTableData& GetTableData() const;
+
+	UE_DEPRECATED(5.3, "FinalValue has moved to TableData. Use property's FinalValue get accessor directly.")
 	float GetFinalValue() const;
 
-	// Set the final value of the WaveTable
+	UE_DEPRECATED(5.3, "FinalValue has moved to TableData and is read-only.")
 	void SetFinalValue(float InValue);
 
-private:
-	UPROPERTY()
-	float FinalValue = WaveTable::InvalidWaveTableValue;
+#if WITH_EDITORONLY_DATA
+	UPROPERTY(meta = (Deprecated = "5.3", DeprecationMessage = "WaveTable data now supports multiple bit depths and has moved to TableData"))
+	TArray<float> WaveTable;
 
+	UPROPERTY(meta = (Deprecated = "5.3", DeprecationMessage = "WaveTable data now supports multiple bit depths and FinalValue has moved to TableData"))
+	float FinalValue;
+#endif // WITH_EDITORONLY_DATA
+
+private:
 	/** Clamps & applies transform to provided values as bipolar signal*/
-	void SampleCurveBipolar(TArrayView<float> InOutValues, float& OutFinalValue) const;
+	void SampleCurveBipolar(FWaveTableData& InOutTableData) const;
 
 	/** Clamps & applies transform to provided values as unipolar signal */
-	void SampleCurveUnipolar(TArrayView<float> InOutValues, float& OutFinalValue) const;
+	void SampleCurveUnipolar(FWaveTableData& InOutTableData) const;
+
+	/** Converts internal TableData to float bit rate if not and re-samples to fit into provided wavetable data. */
+	void SampleWaveTableData(FWaveTableData& InOutTableData) const;
+
+	friend class FWaveTableBankAssetProxy;
+	friend class UWaveTableBank;
 };
diff --git a/Engine/Plugins/Runtime/WaveTable/Source/WaveTable/WaveTable.Build.cs b/Engine/Plugins/Runtime/WaveTable/Source/WaveTable/WaveTable.Build.cs
--- a/Engine/Plugins/Runtime/WaveTable/Source/WaveTable/WaveTable.Build.cs
+++ b/Engine/Plugins/Runtime/WaveTable/Source/WaveTable/WaveTable.Build.cs
@@ -10,15 +10,15 @@ namespace UnrealBuildTool.Rules
 				new string[]
 				{
 					"AudioExtensions",
-					"Core"
+					"Core",
+					"CoreUObject",
+					"Engine"
 				}
 			);
 
 			PrivateDependencyModuleNames.AddRange(
 				new string[]
 				{
-					"CoreUObject",
-					"Engine",
 					"SignalProcessing"
 				}
 			);
diff --git a/Engine/Plugins/Runtime/WaveTable/Source/WaveTableEditor/Private/WaveTableBankEditor.cpp b/Engine/Plugins/Runtime/WaveTable/Source/WaveTableEditor/Private/WaveTableBankEditor.cpp
--- a/Engine/Plugins/Runtime/WaveTable/Source/WaveTableEditor/Private/WaveTableBankEditor.cpp
+++ b/Engine/Plugins/Runtime/WaveTable/Source/WaveTableEditor/Private/WaveTableBankEditor.cpp
@@ -19,630 +19,729 @@
 #define LOCTEXT_NAMESPACE "WaveTableEditor"
 
 
-
-namespace WaveTable
+namespace WaveTable::Editor
 {
-	namespace Editor
+	namespace BankEditorPrivate
 	{
-		namespace BankEditorPrivate
-		{
-			// Two seconds of samples at 48kHz
-			constexpr int32 MaxLiveEditSamples = 48000 * 2;
+		// Two seconds of samples at 48kHz
+		constexpr int32 MaxLiveEditSamples = 48000 * 2;
+		constexpr double MinGutter = -0.05;
+		constexpr double MaxGutter = 0.05;
 
-			int32 GetDrawPointCount(EWaveTableResolution Resolution, const FWaveTableTransform& InTransform)
+		int32 GetDrawPointCount(EWaveTableSamplingMode InMode, int32 InSampleRate, EWaveTableResolution InResolution, const FWaveTableTransform& InTransform)
+		{
+			auto ResolutionToSampleCount = [](EWaveTableResolution InResolution) { return 1 << static_cast<int32>(InResolution); };
+			float NumSamples = 0;
+			if (InMode == EWaveTableSamplingMode::FixedResolution)
 			{
-				// Min/max number of curve points for mathematical
-				// expressions/files drawn on curve stack view.
-				static const FVector2D ExpressionCurvePointCountLimits = { 2, 4096 };
+				EWaveTableResolution DrawPointRes = EWaveTableResolution::None;
+				switch (InResolution)
+				{
+					case EWaveTableResolution::None:
+					{
+						DrawPointRes = EWaveTableResolution::Res_512;
+					}
+					break;
 
-				float WaveTableSize = WaveTable::GetWaveTableSize(Resolution, InTransform.Curve);
-				return FMath::Clamp(WaveTableSize, ExpressionCurvePointCountLimits.X, ExpressionCurvePointCountLimits.Y);
-			}
-		} // namespace BankEditorPrivate
+					case EWaveTableResolution::Maximum:
+					{
+						DrawPointRes = EWaveTableResolution::Res_4096;
+					}
+					break;
 
-		const FName FBankEditorBase::AppIdentifier(TEXT("WaveTableEditorApp"));
-		const FName FBankEditorBase::CurveTabId(TEXT("WaveTableEditor_Curves"));
-		const FName FBankEditorBase::PropertiesTabId(TEXT("WaveTableEditor_Properties"));
+					default:
+					{
+						DrawPointRes = InResolution;
+					}
+					break;
+				};
 
-		FBankEditorBase::FBankEditorBase()
-		{
+				NumSamples = ResolutionToSampleCount(DrawPointRes);
+			}
+			else
+			{
+				if (InTransform.Curve == EWaveTableCurve::File)
+				{
+					NumSamples = InTransform.GetTableData().GetNumSamples();
+				}
+				else
+				{
+					NumSamples = static_cast<int32>(InTransform.GetDuration() * InSampleRate);
+				}
+				NumSamples = FMath::Clamp(NumSamples, 2, ResolutionToSampleCount(EWaveTableResolution::Res_Max));
+			}
+
+			return NumSamples;
 		}
+	} // namespace BankEditorPrivate
 
-		void FBankEditorBase::RegisterTabSpawners(const TSharedRef<FTabManager>& InTabManager)
-		{
-			WorkspaceMenuCategory = InTabManager->AddLocalWorkspaceMenuCategory(LOCTEXT("WorkspaceMenu_WaveTableEditor", "WaveTable Editor"));
+	const FName FBankEditorBase::AppIdentifier(TEXT("WaveTableEditorApp"));
+	const FName FBankEditorBase::CurveTabId(TEXT("WaveTableEditor_Curves"));
+	const FName FBankEditorBase::PropertiesTabId(TEXT("WaveTableEditor_Properties"));
 
-			FAssetEditorToolkit::RegisterTabSpawners(InTabManager);
+	FBankEditorBase::FBankEditorBase()
+	{
+	}
 
-			InTabManager->RegisterTabSpawner(PropertiesTabId, FOnSpawnTab::CreateSP(this, &FBankEditorBase::SpawnTab_Properties))
-				.SetDisplayName(LOCTEXT("DetailsTab", "Details"))
-				.SetGroup(WorkspaceMenuCategory.ToSharedRef())
-				.SetIcon(FSlateIcon(FAppStyle::GetAppStyleSetName(), "LevelEditor.Tabs.Details"));
+	void FBankEditorBase::RegisterTabSpawners(const TSharedRef<FTabManager>& InTabManager)
+	{
+		WorkspaceMenuCategory = InTabManager->AddLocalWorkspaceMenuCategory(LOCTEXT("WorkspaceMenu_WaveTableEditor", "WaveTable Editor"));
 
-			FSlateIcon CurveIcon(FAppStyle::GetAppStyleSetName(), "WaveTableEditor.Tabs.Properties");
-			InTabManager->RegisterTabSpawner(CurveTabId, FOnSpawnTab::CreateLambda([this](const FSpawnTabArgs& Args) { return SpawnTab_OutputCurve(Args); }))
-				.SetDisplayName(LOCTEXT("TransformCurvesTab", "Transform Curves"))
-				.SetGroup(WorkspaceMenuCategory.ToSharedRef())
-				.SetIcon(CurveIcon);
-		}
+		FAssetEditorToolkit::RegisterTabSpawners(InTabManager);
 
-		void FBankEditorBase::UnregisterTabSpawners(const TSharedRef<FTabManager>& InTabManager)
-		{
-			InTabManager->UnregisterTabSpawner(PropertiesTabId);
-			InTabManager->UnregisterTabSpawner(CurveTabId);
-		}
+		InTabManager->RegisterTabSpawner(PropertiesTabId, FOnSpawnTab::CreateSP(this, &FBankEditorBase::SpawnTab_Properties))
+			.SetDisplayName(LOCTEXT("DetailsTab", "Details"))
+			.SetGroup(WorkspaceMenuCategory.ToSharedRef())
+			.SetIcon(FSlateIcon(FAppStyle::GetAppStyleSetName(), "LevelEditor.Tabs.Details"));
 
-		void FBankEditorBase::Init(const EToolkitMode::Type Mode, const TSharedPtr<IToolkitHost>& InitToolkitHost, UObject* InParentObject)
-		{
-			check(InParentObject);
+		FSlateIcon CurveIcon(FAppStyle::GetAppStyleSetName(), "WaveTableEditor.Tabs.Properties");
+		InTabManager->RegisterTabSpawner(CurveTabId, FOnSpawnTab::CreateLambda([this](const FSpawnTabArgs& Args) { return SpawnTab_OutputCurve(Args); }))
+			.SetDisplayName(LOCTEXT("TransformCurvesTab", "Transform Curves"))
+			.SetGroup(WorkspaceMenuCategory.ToSharedRef())
+			.SetIcon(CurveIcon);
+	}
 
-			CurveEditor = MakeShared<FCurveEditor>();
-			FCurveEditorInitParams InitParams;
-			CurveEditor->InitCurveEditor(InitParams);
-			CurveEditor->GridLineLabelFormatXAttribute = LOCTEXT("GridXLabelFormat", "{0}");
+	void FBankEditorBase::UnregisterTabSpawners(const TSharedRef<FTabManager>& InTabManager)
+	{
+		InTabManager->UnregisterTabSpawner(PropertiesTabId);
+		InTabManager->UnregisterTabSpawner(CurveTabId);
+	}
 
-			TUniquePtr<ICurveEditorBounds> EditorBounds = MakeUnique<FStaticCurveEditorBounds>();
-			EditorBounds->SetInputBounds(0.05, 1.05);
-			CurveEditor->SetBounds(MoveTemp(EditorBounds));
+	void FBankEditorBase::Init(const EToolkitMode::Type Mode, const TSharedPtr<IToolkitHost>& InitToolkitHost, UObject* InParentObject)
+	{
+		check(InParentObject);
 
-			CurvePanel = SNew(SCurveEditorPanel, CurveEditor.ToSharedRef());
+		CurveEditor = MakeShared<FCurveEditor>();
+		FCurveEditorInitParams InitParams;
+		CurveEditor->InitCurveEditor(InitParams);
+		CurveEditor->GridLineLabelFormatXAttribute = LOCTEXT("GridXLabelFormat", "{0}");
 
-			// Support undo/redo
-			InParentObject->SetFlags(RF_Transactional);
-			GEditor->RegisterForUndo(this);
+		TUniquePtr<ICurveEditorBounds> EditorBounds = MakeUnique<FStaticCurveEditorBounds>();
+		EditorBounds->SetInputBounds(BankEditorPrivate::MinGutter, 1.0 + BankEditorPrivate::MaxGutter);
+		CurveEditor->SetBounds(MoveTemp(EditorBounds));
 
-			FDetailsViewArgs Args;
-			Args.bHideSelectionTip = true;
-			Args.NotifyHook = this;
+		CurvePanel = SNew(SCurveEditorPanel, CurveEditor.ToSharedRef());
 
-			FPropertyEditorModule& PropertyModule = FModuleManager::LoadModuleChecked<FPropertyEditorModule>("PropertyEditor");
-			PropertiesView = PropertyModule.CreateDetailView(Args);
-			PropertiesView->SetObject(InParentObject);
+		// Support undo/redo
+		InParentObject->SetFlags(RF_Transactional);
+		GEditor->RegisterForUndo(this);
 
-			TSharedRef<FTabManager::FLayout> StandaloneDefaultLayout = FTabManager::NewLayout("Standalone_WaveTableEditor_Layout_v1")
-			->AddArea
+		FDetailsViewArgs Args;
+		Args.bHideSelectionTip = true;
+		Args.NotifyHook = this;
+
+		FPropertyEditorModule& PropertyModule = FModuleManager::LoadModuleChecked<FPropertyEditorModule>("PropertyEditor");
+		PropertiesView = PropertyModule.CreateDetailView(Args);
+		PropertiesView->SetObject(InParentObject);
+
+		TSharedRef<FTabManager::FLayout> StandaloneDefaultLayout = FTabManager::NewLayout("Standalone_WaveTableEditor_Layout_v1")
+		->AddArea
+		(
+			FTabManager::NewPrimaryArea()
+			->SetOrientation(Orient_Vertical)
+			->Split
 			(
-				FTabManager::NewPrimaryArea()
-				->SetOrientation(Orient_Vertical)
+				FTabManager::NewSplitter()
+				->SetSizeCoefficient(0.9f)
+				->SetOrientation(Orient_Horizontal)
+				->Split
+				(
+					FTabManager::NewStack()
+					->SetSizeCoefficient(0.225f)
+					->AddTab(PropertiesTabId, ETabState::OpenedTab)
+				)
 				->Split
 				(
 					FTabManager::NewSplitter()
-					->SetSizeCoefficient(0.9f)
-					->SetOrientation(Orient_Horizontal)
+					->SetSizeCoefficient(0.775f)
+					->SetOrientation(Orient_Vertical)
 					->Split
 					(
 						FTabManager::NewStack()
-						->SetSizeCoefficient(0.225f)
-						->AddTab(PropertiesTabId, ETabState::OpenedTab)
-					)
-					->Split
-					(
-						FTabManager::NewSplitter()
-						->SetSizeCoefficient(0.775f)
-						->SetOrientation(Orient_Vertical)
-						->Split
-						(
-							FTabManager::NewStack()
-							->SetHideTabWell(true)
-							->SetSizeCoefficient(0.33f)
-							->AddTab(CurveTabId, ETabState::OpenedTab)
-						)
+						->SetHideTabWell(true)
+						->SetSizeCoefficient(0.33f)
+						->AddTab(CurveTabId, ETabState::OpenedTab)
 					)
 				)
-			);
-
-			const bool bCreateDefaultStandaloneMenu = true;
-			const bool bCreateDefaultToolbar = true;
-			const bool bToolbarFocusable = false;
-			const bool bUseSmallIcons = true;
-			FAssetEditorToolkit::InitAssetEditor(
-				Mode,
-				InitToolkitHost,
-				AppIdentifier,
-				StandaloneDefaultLayout,
-				bCreateDefaultStandaloneMenu,
-				bCreateDefaultToolbar,
-				InParentObject,
-				bToolbarFocusable,
-				bUseSmallIcons);
-
-			AddToolbarExtender(CurvePanel->GetToolbarExtender());
-
-			if (CurveEditor.IsValid())
-			{
-				RegenerateMenusAndToolbars();
-			}
+			)
+		);
+
+		const bool bCreateDefaultStandaloneMenu = true;
+		const bool bCreateDefaultToolbar = true;
+		const bool bToolbarFocusable = false;
+		const bool bUseSmallIcons = true;
+		FAssetEditorToolkit::InitAssetEditor(
+			Mode,
+			InitToolkitHost,
+			AppIdentifier,
+			StandaloneDefaultLayout,
+			bCreateDefaultStandaloneMenu,
+			bCreateDefaultToolbar,
+			InParentObject,
+			bToolbarFocusable,
+			bUseSmallIcons);
+
+		AddToolbarExtender(CurvePanel->GetToolbarExtender());
+
+		if (CurveEditor.IsValid())
+		{
+			RegenerateMenusAndToolbars();
 		}
+	}
 
-		void FBankEditorBase::ClearExpressionCurve(int32 InTransformIndex)
+	void FBankEditorBase::ClearExpressionCurve(int32 InTransformIndex)
+	{
+		if (CurveData.IsValidIndex(InTransformIndex))
 		{
-			if (CurveData.IsValidIndex(InTransformIndex))
-			{
-				CurveData[InTransformIndex].ExpressionCurve.Reset();
-			}
+			CurveData[InTransformIndex].ExpressionCurve.Reset();
 		}
+	}
 
-		void FBankEditorBase::GenerateExpressionCurve(WaveTable::Editor::FBankEditorBase::FCurveData& OutCurveData, int32 InTransformIndex, bool bInIsUnset)
+	void FBankEditorBase::GenerateExpressionCurve(WaveTable::Editor::FBankEditorBase::FCurveData& OutCurveData, int32 InTransformIndex, bool bInIsUnset)
+	{
+		TSharedPtr<FRichCurve>& Curve = OutCurveData.ExpressionCurve;
+		if (!Curve.IsValid())
 		{
-			TSharedPtr<FRichCurve>& Curve = OutCurveData.ExpressionCurve;
-			if (!Curve.IsValid())
-			{
-				Curve = MakeShared<FRichCurve>();
-			}
+			Curve = MakeShared<FRichCurve>();
+		}
 
-			const bool bIsBipolar = GetBankIsBipolar();
+		const bool bIsBipolar = GetBankIsBipolar();
 
-			if (!GetIsPropertyEditorDisabled())
+		if (!GetIsPropertyEditorDisabled())
+		{
+			if (const FWaveTableTransform* Transform = GetTransform(InTransformIndex))
 			{
-				if (const FWaveTableTransform* Transform = GetTransform(InTransformIndex))
+				const EWaveTableResolution BankResolution = GetBankResolution();
+				const EWaveTableSamplingMode BankSamplingMode = GetBankSamplingMode();
+				const int32 BankSampleRate = GetBankSampleRate();
+				const int32 DrawPointCount = BankEditorPrivate::GetDrawPointCount(BankSamplingMode, BankSampleRate, BankResolution, *Transform);
+
+				FWaveTableData KeyData(EWaveTableBitDepth::IEEE_Float);
+				KeyData.Zero(DrawPointCount);
+				KeyData.SetFinalValue(0.0f);
+
+				// Optimization for really big source files if set to 'File'. Do not recreate WaveTable on the fly
+				// for large samples as caching mechanism is too slow for rapid regeneration (i.e. when making
+				// interactive edits).
+				if (Transform->WaveTableSettings.SourceData.GetNumSamples() < BankEditorPrivate::MaxLiveEditSamples)
 				{
-					// Need at least first and last point to generate expression curve
-					const EWaveTableResolution BankResolution = GetBankResolution();
+					Transform->CreateWaveTable(KeyData, bIsBipolar);
+				}
+				else
+				{
+					Transform->CopyToWaveTable(KeyData, bIsBipolar);
+				}
 
-					const int32 PointCount = BankEditorPrivate::GetDrawPointCount(BankResolution, *Transform);
-					TArray<float> KeyTable;
-					KeyTable.AddZeroed(PointCount);
+				const float Delta = 1.0 / KeyData.GetNumSamples();
+				TArrayView<float> KeyTable;
+				ensureAlways(KeyData.GetDataView(KeyTable)); // Format set above; should be success & always return float view
 
-					float FinalValue = 0.f;
+				FRichCurve& CurveRef = *Curve.Get();
+				TArray<FRichCurveKey> NewKeys;
 
-					// Optimization for really big source files if set to 'File'. Do not recreate WaveTable on the fly
-					// for large samples as caching mechanism is too slow for rapid regeneration (i.e. when making
-					// interactive edits).
-					if (Transform->WaveTableSettings.SourcePCMData.Num() < BankEditorPrivate::MaxLiveEditSamples)
-					{
-						Transform->CreateWaveTable(KeyTable, FinalValue, bIsBipolar);
-					}
-					else
+				float MaxInputValue = 1.0f;
+				switch (BankSamplingMode)
+				{
+					case EWaveTableSamplingMode::FixedResolution:
 					{
-						Transform->CopyToWaveTable(KeyTable, FinalValue, bIsBipolar);
+						for (int32 i = 0; i < KeyTable.Num(); ++i)
+						{
+							NewKeys.Add({ Delta * i, KeyTable[i] });
+						}
 					}
+					break;
 
-					const float Delta = 1.0 / KeyTable.Num();
-
-					FRichCurve& CurveRef = *Curve.Get();
-					TArray<FRichCurveKey> NewKeys;
-
-					// Terminate curves at initial point for wavetables
-					for (int32 i = 0; i < KeyTable.Num(); ++i)
+					case EWaveTableSamplingMode::FixedSampleRate:
 					{
-						NewKeys.Add({ Delta * i, KeyTable[i] });
+						MaxInputValue = Transform->GetDuration();
+						if (Transform->Curve == EWaveTableCurve::File)
+						{
+							MaxInputValue = Transform->GetTableData().GetNumSamples() / (float)BankSampleRate;
+						}
+						for (int32 i = 0; i < KeyTable.Num(); ++i)
+						{
+							const float Time = (Delta * i) * MaxInputValue;
+							NewKeys.Add({ Time, KeyTable[i] });
+						}
 					}
+					break;
 
-					// Terminate curves at initial point for wavetables
-					// (as they loop) and at the final computed location
-					// for non-wavetables (as they are envelopes).
-					if (BankResolution == EWaveTableResolution::None)
-					{
-						float TermPoint = 1.0f;
-						Transform->Apply(TermPoint, bIsBipolar);
-						NewKeys.Add({ 1.0f, TermPoint });
-					}
-					else
+					default:
 					{
-						NewKeys.Add({ 1.0f, KeyTable[0] });
+						static_assert(static_cast<int32>(EWaveTableSamplingMode::COUNT) == 2, "Possible missing switch case coverage for 'EWaveTableSamplingMode'");
+						checkNoEntry();
 					}
+				}
 
-					CurveRef.SetKeys(NewKeys);
-
-					const EWaveTableCurveSource Source = bInIsUnset ? EWaveTableCurveSource::Unset : EWaveTableCurveSource::Expression;
-					SetCurve(InTransformIndex, *Curve.Get(), Source);
+				// Terminate curves at initial point for wavetables
+				// (as they loop) and at the final computed location
+				// for non-wavetables (as they are envelopes).
+				if (BankResolution == EWaveTableResolution::None)
+				{
+					float TermPoint = 1.0f;
+					Transform->Apply(TermPoint, bIsBipolar);
+					NewKeys.Add({ MaxInputValue, TermPoint });
+				}
+				else
+				{
+					NewKeys.Add({ MaxInputValue, KeyTable[0] });
 				}
-			}
-		}
 
-		bool FBankEditorBase::RequiresNewCurve(int32 InTransformIndex, const FRichCurve& InRichCurve) const
-		{
-			const FCurveModelID CurveModelID = CurveData[InTransformIndex].ModelID;
-			const TUniquePtr<FCurveModel>* CurveModel = CurveEditor->GetCurves().Find(CurveModelID);
-			if (!CurveModel || !CurveModel->IsValid())
-			{
-				return true;
-			}
+				CurveRef.SetKeys(NewKeys);
 
-			FWaveTableCurveModel* PatchCurveModel = static_cast<FWaveTableCurveModel*>(CurveModel->Get());
-			check(PatchCurveModel);
-			if (&PatchCurveModel->GetRichCurve() != &InRichCurve)
-			{
-				return true;
+				const EWaveTableCurveSource Source = bInIsUnset ? EWaveTableCurveSource::Unset : EWaveTableCurveSource::Expression;
+				SetCurve(InTransformIndex, *Curve.Get(), Source);
 			}
-
-			return false;
 		}
+	}
 
-		void FBankEditorBase::SetCurve(int32 InTransformIndex, FRichCurve& InRichCurve, EWaveTableCurveSource InSource)
+	bool FBankEditorBase::RequiresNewCurve(int32 InTransformIndex, const FRichCurve& InRichCurve) const
+	{
+		const FCurveModelID CurveModelID = CurveData[InTransformIndex].ModelID;
+		const TUniquePtr<FCurveModel>* CurveModel = CurveEditor->GetCurves().Find(CurveModelID);
+		if (!CurveModel || !CurveModel->IsValid())
 		{
-			check(CurveEditor.IsValid());
-
-			if (!ensure(CurveData.IsValidIndex(InTransformIndex)))
-			{
-				return;
-			}
+			return true;
+		}
 
-			FWaveTableTransform* Transform = GetTransform(InTransformIndex);
-			if (!ensure(Transform))
-			{
-				return;
-			}
+		FWaveTableCurveModel* PatchCurveModel = static_cast<FWaveTableCurveModel*>(CurveModel->Get());
+		check(PatchCurveModel);
+		if (&PatchCurveModel->GetRichCurve() != &InRichCurve)
+		{
+			return true;
+		}
 
-			FCurveData& CurveDataEntry = CurveData[InTransformIndex];
+		return false;
+	}
 
-			const bool bIsBipolar = GetBankIsBipolar();
-			const bool bRequiresNewCurve = RequiresNewCurve(InTransformIndex, InRichCurve);
-			if (bRequiresNewCurve)
-			{
-				TUniquePtr<FWaveTableCurveModel> NewCurve = ConstructCurveModel(InRichCurve, GetEditingObject(), InSource);
-				NewCurve->Refresh(*Transform, InTransformIndex, bIsBipolar);
-				CurveDataEntry.ModelID = CurveEditor->AddCurve(MoveTemp(NewCurve));
-				CurveEditor->PinCurve(CurveDataEntry.ModelID);
-			}
-			else
-			{
-				const TUniquePtr<FCurveModel>& CurveModel = CurveEditor->GetCurves().FindChecked(CurveDataEntry.ModelID);
-				check(CurveModel.Get());
-				static_cast<FWaveTableCurveModel*>(CurveModel.Get())->Refresh(*Transform, InTransformIndex, bIsBipolar);
-			}
-		}
+	void FBankEditorBase::SetCurve(int32 InTransformIndex, FRichCurve& InRichCurve, EWaveTableCurveSource InSource)
+	{
+		check(CurveEditor.IsValid());
 
-		FName FBankEditorBase::GetToolkitFName() const
+		if (!ensure(CurveData.IsValidIndex(InTransformIndex)))
 		{
-			return FName("WaveTableEditor");
+			return;
 		}
 
-		FText FBankEditorBase::GetBaseToolkitName() const
+		FWaveTableTransform* Transform = GetTransform(InTransformIndex);
+		if (!ensure(Transform))
 		{
-			return LOCTEXT( "AppLabel", "WaveTable Editor" );
+			return;
 		}
 
-		FString FBankEditorBase::GetWorldCentricTabPrefix() const
+		FCurveData& CurveDataEntry = CurveData[InTransformIndex];
+
+		const bool bIsBipolar = GetBankIsBipolar();
+		const bool bRequiresNewCurve = RequiresNewCurve(InTransformIndex, InRichCurve);
+		const EWaveTableSamplingMode SamplingMode = GetBankSamplingMode();
+		if (bRequiresNewCurve)
 		{
-			return LOCTEXT("WorldCentricTabPrefix", "WaveTable ").ToString();
+			TUniquePtr<FWaveTableCurveModel> NewCurve = ConstructCurveModel(InRichCurve, GetEditingObject(), InSource);
+			NewCurve->Refresh(*Transform, InTransformIndex, bIsBipolar, SamplingMode);
+			CurveDataEntry.ModelID = CurveEditor->AddCurve(MoveTemp(NewCurve));
+			CurveEditor->PinCurve(CurveDataEntry.ModelID);
 		}
-
-		FLinearColor FBankEditorBase::GetWorldCentricTabColorScale() const
+		else
 		{
-			return FLinearColor( 0.0f, 0.0f, 0.2f, 0.5f );
+			const TUniquePtr<FCurveModel>& CurveModel = CurveEditor->GetCurves().FindChecked(CurveDataEntry.ModelID);
+			check(CurveModel.Get());
+			static_cast<FWaveTableCurveModel*>(CurveModel.Get())->Refresh(*Transform, InTransformIndex, bIsBipolar, SamplingMode);
 		}
+	}
+
+	FName FBankEditorBase::GetToolkitFName() const
+	{
+		return FName("WaveTableEditor");
+	}
+
+	FText FBankEditorBase::GetBaseToolkitName() const
+	{
+		return LOCTEXT( "AppLabel", "WaveTable Editor" );
+	}
+
+	FString FBankEditorBase::GetWorldCentricTabPrefix() const
+	{
+		return LOCTEXT("WorldCentricTabPrefix", "WaveTable ").ToString();
+	}
+
+	FLinearColor FBankEditorBase::GetWorldCentricTabColorScale() const
+	{
+		return FLinearColor( 0.0f, 0.0f, 0.2f, 0.5f );
+	}
+
+	EOrientation FBankEditorBase::GetSnapLabelOrientation() const
+	{
+		return FMultiBoxSettings::UseSmallToolBarIcons.Get()
+			? EOrientation::Orient_Horizontal
+			: EOrientation::Orient_Vertical;
+	}
 
-		EOrientation FBankEditorBase::GetSnapLabelOrientation() const
+	void FBankEditorBase::NotifyPostChange(const FPropertyChangedEvent& PropertyChangedEvent, FProperty* PropertyThatChanged)
+	{
+		const bool bIsInteractive = PropertyChangedEvent.ChangeType == EPropertyChangeType::Interactive;
+		if (bIsInteractive)
 		{
-			return FMultiBoxSettings::UseSmallToolBarIcons.Get()
-				? EOrientation::Orient_Horizontal
-				: EOrientation::Orient_Vertical;
+			RegenerateFileCurves();
 		}
-
-		void FBankEditorBase::NotifyPostChange(const FPropertyChangedEvent& PropertyChangedEvent, FProperty* PropertyThatChanged)
+		else
 		{
-			if (PropertyChangedEvent.ChangeType != EPropertyChangeType::Interactive)
-			{
-				RefreshCurves();
-			}
-			else
+			RefreshCurves();
+
+			if (CurveEditor.IsValid())
 			{
-				// TODO: Isolate file that changed and only update that curve for performance
-				// (Could get bad for larger banks with many file curves).
-				RegenerateFileCurves();
+				double MinBound = 0.0;
+				double MaxBound = 0.0;
+				for (const TPair<FCurveModelID, TUniquePtr<FCurveModel>>& Pair : CurveEditor->GetCurves())
+				{
+					double MinTime = 0.0;
+					double MaxTime = 1.0;
+					Pair.Value->GetTimeRange(MinTime, MaxTime);
+					MinBound = FMath::Min(MinBound, MinTime);
+					MaxBound = FMath::Max(MaxBound, MaxTime);
+				}
+				CurveEditor->GetBounds().SetInputBounds(MinBound - BankEditorPrivate::MinGutter, MaxBound + BankEditorPrivate::MaxGutter);
 			}
 		}
+	}
 
-		TSharedRef<SDockTab> FBankEditorBase::SpawnTab_Properties(const FSpawnTabArgs& Args)
-		{
-			check(Args.GetTabId() == PropertiesTabId);
+	TSharedRef<SDockTab> FBankEditorBase::SpawnTab_Properties(const FSpawnTabArgs& Args)
+	{
+		check(Args.GetTabId() == PropertiesTabId);
 
-			return SNew(SDockTab)
-				.Label(LOCTEXT("SoundWaveTableDetailsTitle", "Details"))
+		return SNew(SDockTab)
+			.Label(LOCTEXT("SoundWaveTableDetailsTitle", "Details"))
+			[
+				PropertiesView.ToSharedRef()
+			];
+	}
+
+	TSharedRef<SDockTab> FBankEditorBase::SpawnTab_OutputCurve(const FSpawnTabArgs& Args)
+	{
+		RefreshCurves();
+		CurveEditor->ZoomToFit();
+
+		TSharedRef<SDockTab> NewDockTab = SNew(SDockTab)
+			.Label(FText::Format(LOCTEXT("WaveTableFilterTitle", "Filter Transform Curve: {0}"), FText::FromString(GetEditingObject()->GetName())))
+			.TabColorScale(GetTabColorScale())
+			[
+				SNew(SBorder)
+				.BorderImage(FAppStyle::GetBrush("ToolPanel.GroupBorder"))
+				.Padding(0.0f)
 				[
-					PropertiesView.ToSharedRef()
-				];
-		}
+					CurvePanel.ToSharedRef()
+				]
+			];
+
+			return NewDockTab;
+	}
 
-		TSharedRef<SDockTab> FBankEditorBase::SpawnTab_OutputCurve(const FSpawnTabArgs& Args)
+	void FBankEditorBase::PostUndo(bool bSuccess)
+	{
+		if (bSuccess)
 		{
 			RefreshCurves();
-			CurveEditor->ZoomToFit();
+		}
+	}
 
-			TSharedRef<SDockTab> NewDockTab = SNew(SDockTab)
-				.Label(FText::Format(LOCTEXT("WaveTableFilterTitle", "Filter Transform Curve: {0}"), FText::FromString(GetEditingObject()->GetName())))
-				.TabColorScale(GetTabColorScale())
-				[
-					SNew(SBorder)
-					.BorderImage(FAppStyle::GetBrush("ToolPanel.GroupBorder"))
-					.Padding(0.0f)
-					[
-						CurvePanel.ToSharedRef()
-					]
-				];
+	void FBankEditorBase::ResetCurves()
+	{
+		check(CurveEditor.IsValid());
 
-				return NewDockTab;
-		}
+		CurveEditor->RemoveAllCurves();
+		CurveData.Reset();
+		CurveData.AddDefaulted(GetNumTransforms());
+	}
 
-		void FBankEditorBase::PostUndo(bool bSuccess)
+	void FBankEditorBase::InitCurves()
+	{
+		for (int32 i = 0; i < GetNumTransforms(); ++i)
 		{
-			if (bSuccess)
+			const FWaveTableTransform* Transform = GetTransform(i);
+			if (!ensure(Transform))
 			{
-				RefreshCurves();
+				continue;
 			}
-		}
 
-		void FBankEditorBase::ResetCurves()
-		{
-			check(CurveEditor.IsValid());
-
-			CurveEditor->RemoveAllCurves();
-			CurveData.Reset();
-			CurveData.AddDefaulted(GetNumTransforms());
-		}
-
-		void FBankEditorBase::InitCurves()
-		{
-			for (int32 i = 0; i < GetNumTransforms(); ++i)
+			switch (Transform->Curve)
 			{
-				const FWaveTableTransform* Transform = GetTransform(i);
-				if (!ensure(Transform))
+				case EWaveTableCurve::Exp:
+				case EWaveTableCurve::Exp_Inverse:
+				case EWaveTableCurve::Linear:
+				case EWaveTableCurve::Linear_Inv:
+				case EWaveTableCurve::Log:
+				case EWaveTableCurve::SCurve:
+				case EWaveTableCurve::Sin:
+				case EWaveTableCurve::Sin_Full:
+				case EWaveTableCurve::File:
 				{
-					continue;
+					if (RequiresNewCurve(i, *CurveData[i].ExpressionCurve.Get()))
+					{
+						ResetCurves();
+					}
 				}
+				break;
 
-				switch (Transform->Curve)
+				case EWaveTableCurve::Shared:
 				{
-					case EWaveTableCurve::Exp:
-					case EWaveTableCurve::Exp_Inverse:
-					case EWaveTableCurve::Linear:
-					case EWaveTableCurve::Linear_Inv:
-					case EWaveTableCurve::Log:
-					case EWaveTableCurve::SCurve:
-					case EWaveTableCurve::Sin:
-					case EWaveTableCurve::Sin_Full:
-					case EWaveTableCurve::File:
+					if (const UCurveFloat* SharedCurve = Transform->CurveShared)
 					{
-						if (RequiresNewCurve(i, *CurveData[i].ExpressionCurve.Get()))
+						if (RequiresNewCurve(i, SharedCurve->FloatCurve))
 						{
 							ResetCurves();
 						}
 					}
-					break;
-
-					case EWaveTableCurve::Shared:
+					else if (RequiresNewCurve(i, *CurveData[i].ExpressionCurve.Get()))
 					{
-						if (const UCurveFloat* SharedCurve = Transform->CurveShared)
-						{
-							if (RequiresNewCurve(i, SharedCurve->FloatCurve))
-							{
-								ResetCurves();
-							}
-						}
-						else if (RequiresNewCurve(i, *CurveData[i].ExpressionCurve.Get()))
-						{
-							ResetCurves();
-						}
+						ResetCurves();
 					}
-					break;
+				}
+				break;
 
-					case EWaveTableCurve::Custom:
+				case EWaveTableCurve::Custom:
+				{
+					if (RequiresNewCurve(i, Transform->CurveCustom))
 					{
-						if (RequiresNewCurve(i, Transform->CurveCustom))
-						{
-							ResetCurves();
-						}
+						ResetCurves();
 					}
-					break;
+				}
+				break;
 
-					default:
-					{
-						static_assert(static_cast<int32>(EWaveTableCurve::Count) == 11, "Possible missing case coverage for output curve.");
-					}
-					break;
+				default:
+				{
+					static_assert(static_cast<int32>(EWaveTableCurve::Count) == 11, "Possible missing case coverage for output curve.");
 				}
+				break;
 			}
 		}
+	}
 
-		void FBankEditorBase::RegenerateFileCurves()
+	void FBankEditorBase::RegenerateFileCurves()
+	{
+		for (int32 i = 0; i < GetNumTransforms(); ++i)
 		{
-			for (int32 i = 0; i < GetNumTransforms(); ++i)
+			if (CurveData.IsValidIndex(i))
 			{
-				if (CurveData.IsValidIndex(i))
-				{
-					FCurveData& CurveDataEntry = CurveData[i];
+				FCurveData& CurveDataEntry = CurveData[i];
 
-					if (FWaveTableTransform* Transform = GetTransform(i))
+				if (FWaveTableTransform* Transform = GetTransform(i))
+				{
+					if (Transform->Curve == EWaveTableCurve::File)
 					{
-						if (Transform->Curve == EWaveTableCurve::File)
-						{
-							GenerateExpressionCurve(CurveDataEntry, i);
-						}
+						GenerateExpressionCurve(CurveDataEntry, i);
 					}
 				}
 			}
 		}
+	}
 
-		void FBankEditorBase::RefreshCurves()
+	void FBankEditorBase::RefreshCurves()
+	{
+		check(CurveEditor.IsValid());
+
+		for (const FCurveData& CurveDataEntry : CurveData)
 		{
-			check(CurveEditor.IsValid());
+			CurveEditor->UnpinCurve(CurveDataEntry.ModelID);
+		}
 
-			for (const FCurveData& CurveDataEntry : CurveData)
-			{
-				CurveEditor->UnpinCurve(CurveDataEntry.ModelID);
-			}
+		if (GetIsPropertyEditorDisabled())
+		{
+			ResetCurves();
+			return;
+		}
 
-			if (GetIsPropertyEditorDisabled())
-			{
-				ResetCurves();
-				return;
-			}
+		const int32 NumCurves = GetNumTransforms();
+		if (NumCurves == CurveData.Num())
+		{
+			InitCurves();
+		}
+		else
+		{
+			ResetCurves();
+		}
 
-			const int32 NumCurves = GetNumTransforms();
-			if (NumCurves == CurveData.Num())
+		for (int32 i = 0; i < GetNumTransforms(); ++i)
+		{
+			if (!CurveData.IsValidIndex(i))
 			{
-				InitCurves();
+				continue;
 			}
-			else
+			FCurveData& CurveDataEntry = CurveData[i];
+
+			FWaveTableTransform* Transform = GetTransform(i);
+			if (!Transform)
 			{
-				ResetCurves();
+				continue;
 			}
 
-			for (int32 i = 0; i < GetNumTransforms(); ++i)
+			switch (Transform->Curve)
 			{
-				if (!CurveData.IsValidIndex(i))
+				case EWaveTableCurve::Exp:
+				case EWaveTableCurve::Exp_Inverse:
+				case EWaveTableCurve::Linear:
+				case EWaveTableCurve::Linear_Inv:
+				case EWaveTableCurve::Log:
+				case EWaveTableCurve::SCurve:
+				case EWaveTableCurve::Sin:
+				case EWaveTableCurve::Sin_Full:
+				case EWaveTableCurve::File:
 				{
-					continue;
-				}
-				FCurveData& CurveDataEntry = CurveData[i];
-
-				FWaveTableTransform* Transform = GetTransform(i);
-				if (!Transform)
-				{
-					continue;
+					GenerateExpressionCurve(CurveDataEntry, i);
 				}
+				break;
 
-				switch (Transform->Curve)
+				case EWaveTableCurve::Shared:
 				{
-					case EWaveTableCurve::Exp:
-					case EWaveTableCurve::Exp_Inverse:
-					case EWaveTableCurve::Linear:
-					case EWaveTableCurve::Linear_Inv:
-					case EWaveTableCurve::Log:
-					case EWaveTableCurve::SCurve:
-					case EWaveTableCurve::Sin:
-					case EWaveTableCurve::Sin_Full:
-					case EWaveTableCurve::File:
+					if (UCurveFloat* SharedCurve = Transform->CurveShared)
 					{
-						GenerateExpressionCurve(CurveDataEntry, i);
-					}
-					break;
-
-					case EWaveTableCurve::Shared:
-					{
-						if (UCurveFloat* SharedCurve = Transform->CurveShared)
-						{
-							ClearExpressionCurve(i);
-							SetCurve(i, SharedCurve->FloatCurve, EWaveTableCurveSource::Shared);
-						}
-						else
-						{
-							// Builds a dummy expression that just maps input to output in case
-							// where asset isn't selected and leave source as unset
-							GenerateExpressionCurve(CurveDataEntry, i, true /* bIsUnset */);
-						}
-					}
-					break;
-
-					case EWaveTableCurve::Custom:
-					{
-						FRichCurve& CustomCurve = Transform->CurveCustom;
-						TrimKeys(CustomCurve);
 						ClearExpressionCurve(i);
-						SetCurve(i, CustomCurve, EWaveTableCurveSource::Custom);
+						SetCurve(i, SharedCurve->FloatCurve, EWaveTableCurveSource::Shared);
 					}
-					break;
-
-					default:
+					else
 					{
-						static_assert(static_cast<int32>(EWaveTableCurve::Count) == 11, "Possible missing case coverage for output curve.");
+						// Builds a dummy expression that just maps input to output in case
+						// where asset isn't selected and leave source as unset
+						GenerateExpressionCurve(CurveDataEntry, i, true /* bIsUnset */);
 					}
-					break;
 				}
-			}
+				break;
 
-			// Collect and remove stale curves from editor
-			TArray<FCurveModelID> ToRemove;
-			TSet<FCurveModelID> ActiveModelIDs;
-			for (const FCurveData& CurveDataEntry : CurveData)
-			{
-				ActiveModelIDs.Add(CurveDataEntry.ModelID);
-			}
-
-			// Remove all dead curves
-			{
-				const TMap<FCurveModelID, TUniquePtr<FCurveModel>>& Curves = CurveEditor->GetCurves();
-				for (const TPair<FCurveModelID, TUniquePtr<FCurveModel>>& Pair : Curves)
+				case EWaveTableCurve::Custom:
 				{
-					if (!ActiveModelIDs.Contains(Pair.Key))
-					{
-						ToRemove.Add(Pair.Key);
-					}
+					FRichCurve& CustomCurve = Transform->CurveCustom;
+					TrimKeys(CustomCurve);
+					ClearExpressionCurve(i);
+					SetCurve(i, CustomCurve, EWaveTableCurveSource::Custom);
 				}
+				break;
 
-				for (FCurveModelID ModelID : ToRemove)
+				default:
 				{
-					CurveEditor->RemoveCurve(ModelID);
+					static_assert(static_cast<int32>(EWaveTableCurve::Count) == 11, "Possible missing case coverage for output curve.");
 				}
-			}
-
-			// Re-pin all curves
-			for (const FCurveData& CurveDataEntry : CurveData)
-			{
-				CurveEditor->PinCurve(CurveDataEntry.ModelID);
+				break;
 			}
 		}
 
-		void FBankEditorBase::PostRedo(bool bSuccess)
+		// Collect and remove stale curves from editor
+		TArray<FCurveModelID> ToRemove;
+		TSet<FCurveModelID> ActiveModelIDs;
+		for (const FCurveData& CurveDataEntry : CurveData)
 		{
-			if (bSuccess)
-			{
-				RefreshCurves();
-			}
+			ActiveModelIDs.Add(CurveDataEntry.ModelID);
 		}
 
-		void FBankEditorBase::TrimKeys(FRichCurve& OutCurve)
+		// Remove all dead curves
 		{
-			while (OutCurve.GetNumKeys() > 0 && 0.0f > OutCurve.GetFirstKey().Time)
+			const TMap<FCurveModelID, TUniquePtr<FCurveModel>>& Curves = CurveEditor->GetCurves();
+			for (const TPair<FCurveModelID, TUniquePtr<FCurveModel>>& Pair : Curves)
 			{
-				OutCurve.DeleteKey(OutCurve.GetFirstKeyHandle());
+				if (!ActiveModelIDs.Contains(Pair.Key))
+				{
+					ToRemove.Add(Pair.Key);
+				}
 			}
 
-			while (OutCurve.GetNumKeys() > 0 && 1.0f < OutCurve.GetLastKey().Time)
+			for (FCurveModelID ModelID : ToRemove)
 			{
-				OutCurve.DeleteKey(OutCurve.GetLastKeyHandle());
+				CurveEditor->RemoveCurve(ModelID);
 			}
 		}
 
-		TUniquePtr<WaveTable::Editor::FWaveTableCurveModel> FBankEditor::ConstructCurveModel(FRichCurve& InRichCurve, UObject* InParentObject, EWaveTableCurveSource InSource)
+		// Re-pin all curves
+		for (const FCurveData& CurveDataEntry : CurveData)
 		{
-			using namespace WaveTable::Editor;
-			return MakeUnique<FWaveTableCurveModel>(InRichCurve, GetEditingObject(), InSource);
+			CurveEditor->PinCurve(CurveDataEntry.ModelID);
 		}
+	}
 
-		EWaveTableResolution FBankEditor::GetBankResolution() const
+	void FBankEditorBase::PostRedo(bool bSuccess)
+	{
+		if (bSuccess)
 		{
-			if (UWaveTableBank* Bank = Cast<UWaveTableBank>(GetEditingObject()))
-			{
-				return Bank->Resolution;
-			}
+			RefreshCurves();
+		}
+	}
 
-			return EWaveTableResolution::None;
+	void FBankEditorBase::TrimKeys(FRichCurve& OutCurve)
+	{
+		while (OutCurve.GetNumKeys() > 0 && 0.0f > OutCurve.GetFirstKey().Time)
+		{
+			OutCurve.DeleteKey(OutCurve.GetFirstKeyHandle());
 		}
 
-		bool FBankEditor::GetBankIsBipolar() const
+		while (OutCurve.GetNumKeys() > 0 && 1.0f < OutCurve.GetLastKey().Time)
 		{
-			if (UWaveTableBank* Bank = Cast<UWaveTableBank>(GetEditingObject()))
-			{
-				return Bank->bBipolar;
-			}
+			OutCurve.DeleteKey(OutCurve.GetLastKeyHandle());
+		}
+	}
 
-			return false;
+	TUniquePtr<WaveTable::Editor::FWaveTableCurveModel> FBankEditor::ConstructCurveModel(FRichCurve& InRichCurve, UObject* InParentObject, EWaveTableCurveSource InSource)
+	{
+		using namespace WaveTable::Editor;
+		return MakeUnique<FWaveTableCurveModel>(InRichCurve, GetEditingObject(), InSource);
+	}
+
+	EWaveTableResolution FBankEditor::GetBankResolution() const
+	{
+		if (UWaveTableBank* Bank = Cast<UWaveTableBank>(GetEditingObject()))
+		{
+			return Bank->Resolution;
 		}
 
-		int32 FBankEditor::GetNumTransforms() const
+		return EWaveTableResolution::None;
+	}
+
+	EWaveTableSamplingMode FBankEditor::GetBankSamplingMode() const
+	{
+		if (UWaveTableBank* Bank = Cast<UWaveTableBank>(GetEditingObject()))
 		{
-			if (UWaveTableBank* Bank = Cast<UWaveTableBank>(GetEditingObject()))
-			{
-				return Bank->Entries.Num();
-			}
+			return Bank->SampleMode;
+		}
+
+		return EWaveTableSamplingMode::FixedResolution;
+	}
 
-			return 0;
+	int32 FBankEditor::GetBankSampleRate() const
+	{
+		if (UWaveTableBank* Bank = Cast<UWaveTableBank>(GetEditingObject()))
+		{
+			return Bank->SampleRate;
 		}
 
-		FWaveTableTransform* FBankEditor::GetTransform(int32 InIndex) const
+		return 48000;
+	}
+
+	bool FBankEditor::GetBankIsBipolar() const
+	{
+		if (UWaveTableBank* Bank = Cast<UWaveTableBank>(GetEditingObject()))
+		{
+			return Bank->bBipolar;
+		}
+
+		return false;
+	}
+
+	int32 FBankEditor::GetNumTransforms() const
+	{
+		if (UWaveTableBank* Bank = Cast<UWaveTableBank>(GetEditingObject()))
+		{
+			return Bank->GetEntries().Num();
+		}
+
+		return 0;
+	}
+
+	FWaveTableTransform* FBankEditor::GetTransform(int32 InIndex) const
+	{
+		if (UWaveTableBank* Bank = Cast<UWaveTableBank>(GetEditingObject()))
 		{
-			if (UWaveTableBank* Bank = Cast<UWaveTableBank>(GetEditingObject()))
+			if (Bank->GetEntries().IsValidIndex(InIndex))
 			{
-				if (Bank->Entries.IsValidIndex(InIndex))
-				{
-					return &Bank->Entries[InIndex].Transform;
-				}
+				return &Bank->GetEntries()[InIndex].Transform;
 			}
-
-			return nullptr;
 		}
-	} // namespace Editor
-} // namespace WaveTable
+
+		return nullptr;
+	}
+} // namespace WaveTable::Editor
 #undef LOCTEXT_NAMESPACE
diff --git a/Engine/Plugins/Runtime/WaveTable/Source/WaveTableEditor/Private/WaveTableCurveEditorViewStacked.cpp b/Engine/Plugins/Runtime/WaveTable/Source/WaveTableEditor/Private/WaveTableCurveEditorViewStacked.cpp
--- a/Engine/Plugins/Runtime/WaveTable/Source/WaveTableEditor/Private/WaveTableCurveEditorViewStacked.cpp
+++ b/Engine/Plugins/Runtime/WaveTable/Source/WaveTableEditor/Private/WaveTableCurveEditorViewStacked.cpp
@@ -25,18 +25,19 @@ namespace WaveTable
 			check(InOwner);
 		}
 
-		void FWaveTableCurveModel::Refresh(const FWaveTableTransform& InTransform, int32 InCurveIndex, bool bInIsBipolar)
+		void FWaveTableCurveModel::Refresh(const FWaveTableTransform& InTransform, int32 InCurveIndex, bool bInIsBipolar, EWaveTableSamplingMode InSamplingMode)
 		{
 			// Must be set prior to remainder of refresh to avoid
 			// child class implementation acting on incorrect cached index
 			// (Editor may re-purpose models resulting in index change).
 			CurveIndex = InCurveIndex;
 
+			CurveDuration = InTransform.GetDuration();
 			Color = GetCurveColor();
+			SamplingMode = InSamplingMode;
 
 			FText OutputAxisName;
 			RefreshCurveDescriptorText(InTransform, ShortDisplayName, InputAxisName, OutputAxisName);
-
 			AxesDescriptor = FText::Format(LOCTEXT("WaveTableCurveDisplayTitle_AxesFormat", "X = {0}, Y = {1}"), InputAxisName, OutputAxisName);
 
 			IntentionName = TEXT("AudioControlValue");
@@ -46,15 +47,15 @@ namespace WaveTable
 
 			bIsBipolar = bInIsBipolar;
 
+			NumSamples = InTransform.GetTableData().GetNumSamples();
+
 			if (InTransform.Curve == EWaveTableCurve::File)
 			{
-				NumSamples = InTransform.WaveTable.Num();
 				FadeInRatio = InTransform.WaveTableSettings.FadeIn;
 				FadeOutRatio = InTransform.WaveTableSettings.FadeOut;
 			}
 			else
 			{
-				NumSamples = 0;
 				FadeInRatio = 0.0f;
 				FadeOutRatio = 0.0f;
 			}
@@ -134,7 +135,28 @@ namespace WaveTable
 			check(Enum);
 
 			Enum->GetDisplayValueAsText(InTransform.Curve, OutShortDisplayName);
-			OutInputAxisName = LOCTEXT("WaveTableCurveModel_DefaultInputAxisName", "Unit Phase (Table Index)");
+
+			switch (SamplingMode)
+			{
+				case EWaveTableSamplingMode::FixedResolution:
+				{
+					OutInputAxisName = LOCTEXT("WaveTableCurveModel_DefaultInputAxisName", "Unit Phase (Table Index)");
+				}
+				break;
+
+				case EWaveTableSamplingMode::FixedSampleRate:
+				{
+					OutInputAxisName = LOCTEXT("WaveTableCurveModel_DefaultInputAxisName", "Time");
+				}
+				break;
+
+				default:
+				{
+					static_assert(static_cast<int32>(EWaveTableSamplingMode::COUNT) == 2, "Possible missing switch case coverage for 'EWaveTableSamplingMode'");
+					checkNoEntry();
+				}
+			};
+
 			OutOutputAxisName = LOCTEXT("WaveTableCurveModel_DefaultOutputAxisName", "Amplitude");
 		}
 
@@ -191,15 +213,15 @@ namespace WaveTable
 			const FCurveEditorScreenSpace ViewSpace = GetViewSpace();
 
 			// Draw the curve labels for each view
-			for (auto It = CurveInfoByID.CreateConstIterator(); It; ++It)
+			for (const TPair<FCurveModelID, SCurveEditorView::FCurveInfo>& Pair : CurveInfoByID)
 			{
-				FCurveModel* Curve = CurveEditor->FindCurve(It.Key());
+				FCurveModel* Curve = CurveEditor->FindCurve(Pair.Key);
 				if (!ensureAlways(Curve))
 				{
 					continue;
 				}
 
-				const int32  CurveIndexFromBottom = CurveInfoByID.Num() - It->Value.CurveIndex - 1;
+				const int32  CurveIndexFromBottom = CurveInfoByID.Num() - Pair.Value.CurveIndex - 1;
 				const double PaddingToBottomOfView = (CurveIndexFromBottom + 1) * ValueSpacePadding;
 
 				const float PixelBottom = ViewSpace.ValueToScreen(CurveIndexFromBottom + PaddingToBottomOfView);
@@ -327,7 +349,7 @@ namespace WaveTable
 					DrawInfo.LinePoints[1].Y = DrawInfo.GetPixelBottom();
 
 					// Draw major vertical grid lines
-					for (int i = 0; i < MajorGridLinesX.Num(); ++i)
+					for (int32 i = 0; i < MajorGridLinesX.Num(); ++i)
 					{
 						const float VerticalLine = FMath::RoundToFloat(MajorGridLinesX[i]);
 						const FText& Label = CurveModelGridLabelsX[i];
@@ -382,16 +404,17 @@ namespace WaveTable
 					};
 
 					const bool bIsBipolar = EditorModel->GetIsBipolar();
+					const float Duration = EditorModel->GetSamplingMode() == EWaveTableSamplingMode::FixedResolution ? 1.0f : EditorModel->GetCurveDuration();
 					float FadeInRatio = EditorModel->GetFadeInRatio();
 					if (FadeInRatio > 0.0f)
 					{
-						DrawFade(0.0f, FadeInRatio, true, bIsBipolar);
+						DrawFade(0.0f, Duration* FadeInRatio, true, bIsBipolar);
 					}
 
 					float FadeOutRatio = EditorModel->GetFadeOutRatio();
 					if (FadeOutRatio > 0.0f)
 					{
-						DrawFade(1.0f, 1.0f - FadeOutRatio, false, bIsBipolar);
+						DrawFade(Duration, Duration * (1.0f - FadeOutRatio), false, bIsBipolar);
 					}
 				}
 			}
@@ -487,16 +510,35 @@ namespace WaveTable
 
 		void SViewStacked::FormatInputLabel(const FWaveTableCurveModel& EditorModel, const FNumberFormattingOptions& InLabelFormat, FText& InOutLabel) const
 		{
-			const int32 NumSamples = EditorModel.GetNumSamples();
-			if (NumSamples != 0)
+			switch (EditorModel.GetSamplingMode())
 			{
-				const float Value = FCString::Atof(*InOutLabel.ToString());
-				if (Value >= 0 && Value <= 1.0f)
+				case EWaveTableSamplingMode::FixedResolution:
 				{
-					const int32 ValueTrunc = FMath::FloorToInt32(NumSamples * Value);
-					InOutLabel = FText::Format(LOCTEXT("WaveTableCurveEditor_InputLabelFormat", "{0} ({1})"), InOutLabel, FText::AsNumber(ValueTrunc));
+					int32 NumSamples = EditorModel.GetNumSamples();
+					if (NumSamples != 0)
+					{
+						const float Value = FCString::Atof(*InOutLabel.ToString());
+						if (Value >= 0 && Value <= 1.0f)
+						{
+							NumSamples = FMath::FloorToInt32(NumSamples * Value);
+							InOutLabel = FText::Format(LOCTEXT("WaveTableCurveEditor_ResolutionInputLabelFormat", "{0} ({1})"), InOutLabel, FText::AsNumber(NumSamples));
+						}
+					}
 				}
-			}
+				break;
+
+				case EWaveTableSamplingMode::FixedSampleRate:
+				{
+					InOutLabel = FText::Format(LOCTEXT("WaveTableCurveEditor_DurationInputLabelFormat", "{0}s"), InOutLabel);
+				}
+				break;
+
+				default:
+				{
+					static_assert(static_cast<int32>(EWaveTableSamplingMode::COUNT) == 2, "Possible missing switch case coverage for 'EWaveTableSamplingMode'");
+					checkNoEntry();
+				}
+			};
 		}
 
 		SViewStacked::FGridDrawInfo::FGridDrawInfo(const FGeometry* InAllottedGeometry, const FCurveEditorScreenSpace& InScreenSpace, FLinearColor InGridColor, int32 InBaseLayerId)
diff --git a/Engine/Plugins/Runtime/WaveTable/Source/WaveTableEditor/Private/WaveTableEditorModule.cpp b/Engine/Plugins/Runtime/WaveTable/Source/WaveTableEditor/Private/WaveTableEditorModule.cpp
--- a/Engine/Plugins/Runtime/WaveTable/Source/WaveTableEditor/Private/WaveTableEditorModule.cpp
+++ b/Engine/Plugins/Runtime/WaveTable/Source/WaveTableEditor/Private/WaveTableEditorModule.cpp
@@ -5,82 +5,122 @@
 #include "CurveDrawInfo.h"
 #include "ICurveEditorModule.h"
 #include "PropertyEditorModule.h"
+#include "Styling/SlateStyleMacros.h"
+#include "Styling/SlateStyleRegistry.h"
 #include "WaveTableCurveEditorViewStacked.h"
 #include "WaveTableTransformLayout.h"
 
 DEFINE_LOG_CATEGORY(LogWaveTableEditor);
 
 
-namespace WaveTable
+namespace WaveTable::Editor
 {
-	namespace Editor
+	namespace ModulePrivate
 	{
-		namespace ModulePrivate
+		static const FName AssetToolName { "AssetTools" };
+		static const FName CurveEditorName { "CurveEditor" };
+		static const FName PropertyEditorName { "PropertyEditor" };
+
+		template <typename T>
+		void AddAssetAction(IAssetTools& AssetTools, TArray<TSharedPtr<FAssetTypeActions_Base>>& AssetArray)
 		{
-			static const FName AssetToolName { "AssetTools" };
-			static const FName CurveEditorName { "CurveEditor" };
-			static const FName PropertyEditorName { "PropertyEditor" };
+			TSharedPtr<T> AssetAction = MakeShared<T>();
+			TSharedPtr<FAssetTypeActions_Base> AssetActionBase = StaticCastSharedPtr<FAssetTypeActions_Base>(AssetAction);
+			AssetTools.RegisterAssetTypeActions(AssetAction.ToSharedRef());
+			AssetArray.Add(AssetActionBase);
+		}
 
-			template <typename T>
-			void AddAssetAction(IAssetTools& AssetTools, TArray<TSharedPtr<FAssetTypeActions_Base>>& AssetArray)
+		class FSlateStyle final : public FSlateStyleSet
+		{
+		public:
+			FSlateStyle()
+				: FSlateStyleSet("WaveTableStyle")
 			{
-				TSharedPtr<T> AssetAction = MakeShared<T>();
-				TSharedPtr<FAssetTypeActions_Base> AssetActionBase = StaticCastSharedPtr<FAssetTypeActions_Base>(AssetAction);
-				AssetTools.RegisterAssetTypeActions(AssetAction.ToSharedRef());
-				AssetArray.Add(AssetActionBase);
-			}
-		} // namespace ModulePrivate
+				SetParentStyleName(FAppStyle::GetAppStyleSetName());
 
-		void FModule::StartupModule()
-		{
-			using namespace ModulePrivate;
+				SetContentRoot(FPaths::EnginePluginsDir() / TEXT("Runtime/WaveTable/Content/Editor/Slate"));
+				SetCoreContentRoot(FPaths::EngineContentDir() / TEXT("Slate"));
 
-			FPropertyEditorModule& PropertyModule = FModuleManager::LoadModuleChecked<FPropertyEditorModule>(PropertyEditorName);
-			PropertyModule.RegisterCustomPropertyTypeLayout
-			(
-				"WaveTableTransform",
-				FOnGetPropertyTypeCustomizationInstance::CreateStatic(&FTransformLayoutCustomization::MakeInstance)
-			);
+				static const FVector2D Icon16 = FVector2D(16.0f, 16.0f);
+				static const FVector2D Icon64 = FVector2D(64.0f, 64.0f);
 
-			ICurveEditorModule& CurveEditorModule = FModuleManager::LoadModuleChecked<ICurveEditorModule>(CurveEditorName);
-			FWaveTableCurveModel::WaveTableViewId = CurveEditorModule.RegisterView(FOnCreateCurveEditorView::CreateStatic(
-				[](TWeakPtr<FCurveEditor> WeakCurveEditor) -> TSharedRef<SCurveEditorView>
+				// WaveTabe Editor
 				{
-					return SNew(SViewStacked, WeakCurveEditor);
+					// Class Icons
+					auto SetClassIcon = [this, InIcon16 = Icon16, InIcon64 = Icon64](const FString& ClassName)
+					{
+						const FString IconFileName = FString::Printf(TEXT("Icons/%s"), *ClassName.ToLower());
+						const FSlateColor DefaultForeground(FStyleColors::Foreground);
+
+						Set(*FString::Printf(TEXT("ClassIcon.%s"), *ClassName), new IMAGE_BRUSH_SVG(IconFileName, InIcon16));
+						Set(*FString::Printf(TEXT("ClassThumbnail.%s"), *ClassName), new IMAGE_BRUSH_SVG(IconFileName, InIcon64));
+					};
+
+					SetClassIcon(TEXT("WaveTableBank"));
 				}
-			));
 
-			IAssetTools& AssetTools = FModuleManager::LoadModuleChecked<FAssetToolsModule>(AssetToolName).Get();
-			AddAssetAction<FAssetTypeActions_WaveTableBank>(AssetTools, AssetActions);
-		}
+				FSlateStyleRegistry::RegisterSlateStyle(*this);
+			}
+		};
+	} // namespace ModulePrivate
 
-		void FModule::ShutdownModule()
-		{
-			using namespace ModulePrivate;
+	void FModule::StartupModule()
+	{
+		using namespace ModulePrivate;
+
+		FPropertyEditorModule& PropertyModule = FModuleManager::LoadModuleChecked<FPropertyEditorModule>(PropertyEditorName);
+		PropertyModule.RegisterCustomPropertyTypeLayout
+		(
+			"WaveTableTransform",
+			FOnGetPropertyTypeCustomizationInstance::CreateStatic(&FTransformLayoutCustomization::MakeInstance)
+		);
 
-			if (FModuleManager::Get().IsModuleLoaded(PropertyEditorName))
+		PropertyModule.RegisterCustomPropertyTypeLayout
+		(
+			"WaveTableData",
+			FOnGetPropertyTypeCustomizationInstance::CreateStatic(&FWaveTableDataLayoutCustomization::MakeInstance)
+		);
+
+		ICurveEditorModule& CurveEditorModule = FModuleManager::LoadModuleChecked<ICurveEditorModule>(CurveEditorName);
+		FWaveTableCurveModel::WaveTableViewId = CurveEditorModule.RegisterView(FOnCreateCurveEditorView::CreateStatic(
+			[](TWeakPtr<FCurveEditor> WeakCurveEditor) -> TSharedRef<SCurveEditorView>
 			{
-				FPropertyEditorModule& PropertyModule = FModuleManager::GetModuleChecked<FPropertyEditorModule>(PropertyEditorName);
-				PropertyModule.UnregisterCustomPropertyTypeLayout("WaveTableTransform");
+				return SNew(SViewStacked, WeakCurveEditor);
 			}
+		));
+
+		IAssetTools& AssetTools = FModuleManager::LoadModuleChecked<FAssetToolsModule>(AssetToolName).Get();
+		AddAssetAction<FAssetTypeActions_WaveTableBank>(AssetTools, AssetActions);
+
+		StyleSet = MakeShared<FSlateStyle>();
+	}
+
+	void FModule::ShutdownModule()
+	{
+		using namespace ModulePrivate;
+
+		if (FModuleManager::Get().IsModuleLoaded(PropertyEditorName))
+		{
+			FPropertyEditorModule& PropertyModule = FModuleManager::GetModuleChecked<FPropertyEditorModule>(PropertyEditorName);
+			PropertyModule.UnregisterCustomPropertyTypeLayout("WaveTableTransform");
+		}
 
-// 			if (FModuleManager::Get().IsModuleLoaded(AssetToolName))
+// 		if (FModuleManager::Get().IsModuleLoaded(AssetToolName))
+// 		{
+// 			IAssetTools& AssetTools = FModuleManager::GetModuleChecked<FAssetToolsModule>(AssetToolName).Get();
+// 			for (const TSharedPtr<FAssetTypeActions_Base>& ActionPtr : AssetActions)
 // 			{
-// 				IAssetTools& AssetTools = FModuleManager::GetModuleChecked<FAssetToolsModule>(AssetToolName).Get();
-// 				for (const TSharedPtr<FAssetTypeActions_Base>& ActionPtr : AssetActions)
-// 				{
-// 					AssetTools.UnregisterAssetTypeActions(ActionPtr);
-// 				}
+// 				AssetTools.UnregisterAssetTypeActions(ActionPtr);
 // 			}
+// 		}
 
-			if (FModuleManager::Get().IsModuleLoaded(CurveEditorName))
-			{
-				ICurveEditorModule& CurveEditorModule = FModuleManager::LoadModuleChecked<ICurveEditorModule>(CurveEditorName);
-				CurveEditorModule.UnregisterView(FWaveTableCurveModel::WaveTableViewId);
-				FWaveTableCurveModel::WaveTableViewId = ECurveEditorViewID::Invalid;
-			}
+		if (FModuleManager::Get().IsModuleLoaded(CurveEditorName))
+		{
+			ICurveEditorModule& CurveEditorModule = FModuleManager::LoadModuleChecked<ICurveEditorModule>(CurveEditorName);
+			CurveEditorModule.UnregisterView(FWaveTableCurveModel::WaveTableViewId);
+			FWaveTableCurveModel::WaveTableViewId = ECurveEditorViewID::Invalid;
 		}
-	} // namespace Editor
-} // namespace WaveTable
+	}
+} // namespace WaveTable::Editor
 
 IMPLEMENT_MODULE(WaveTable::Editor::FModule, WaveTableEditor);
diff --git a/Engine/Plugins/Runtime/WaveTable/Source/WaveTableEditor/Private/WaveTableEditorModule.h b/Engine/Plugins/Runtime/WaveTable/Source/WaveTableEditor/Private/WaveTableEditorModule.h
--- a/Engine/Plugins/Runtime/WaveTable/Source/WaveTableEditor/Private/WaveTableEditorModule.h
+++ b/Engine/Plugins/Runtime/WaveTable/Source/WaveTableEditor/Private/WaveTableEditorModule.h
@@ -1,29 +1,27 @@
 // Copyright Epic Games, Inc. All Rights Reserved.
 #pragma once
 
+#include "AssetTypeActions_Base.h"
 #include "Modules/ModuleInterface.h"
+#include "Styling/SlateStyle.h"
 #include "Templates/SharedPointer.h"
 
-class FAssetTypeActions_Base;
-
-
 DECLARE_LOG_CATEGORY_EXTERN(LogWaveTableEditor, Log, All);
 
-namespace WaveTable
+
+namespace WaveTable::Editor
 {
-	namespace Editor
+	class FModule : public IModuleInterface
 	{
-		class FModule : public IModuleInterface
-		{
-		public:
-			FModule() = default;
+	public:
+		FModule() = default;
 
-			virtual void StartupModule() override;
+		virtual void StartupModule() override;
 
-			virtual void ShutdownModule() override;
+		virtual void ShutdownModule() override;
 
-		private:
-			TArray<TSharedPtr<FAssetTypeActions_Base>> AssetActions;
-		};
-	} // namespace Editor
-} // namespace WaveTable
+	private:
+		TArray<TSharedPtr<FAssetTypeActions_Base>> AssetActions;
+		TSharedPtr<FSlateStyleSet> StyleSet;
+	};
+} // namespace WaveTable::Editor
diff --git a/Engine/Plugins/Runtime/WaveTable/Source/WaveTableEditor/Private/WaveTableFileUtilities.cpp b/Engine/Plugins/Runtime/WaveTable/Source/WaveTableEditor/Private/WaveTableFileUtilities.cpp
--- a/Engine/Plugins/Runtime/WaveTable/Source/WaveTableEditor/Private/WaveTableFileUtilities.cpp
+++ b/Engine/Plugins/Runtime/WaveTable/Source/WaveTableEditor/Private/WaveTableFileUtilities.cpp
@@ -5,51 +5,136 @@
 #include "Sound/SoundWave.h"
 #include "UObject/Package.h"
 #include "Utils.h"
+#include "WaveTableEditorModule.h"
+#include "WaveTableSampler.h"
 
 
-namespace WaveTable
+namespace WaveTable::Editor
 {
-	namespace Editor
+	namespace FileUtilities
 	{
-		namespace FileUtilities
+		void LoadPCMChannel(const FString& InFilePath, int32 InChannelIndex, FWaveTableData& OutData, int32& OutSampleRate)
 		{
-			void LoadPCMChannel(const FString& InFilePath, int32 InChannelIndex, TArray<float>& OutPCMData)
-			{
-				OutPCMData.Empty();
+			OutData.Empty();
 
-				if (!InFilePath.IsEmpty())
+			if (!InFilePath.IsEmpty())
+			{
+				if (USoundFactory* SoundWaveFactory = NewObject<USoundFactory>())
 				{
-					if (USoundFactory* SoundWaveFactory = NewObject<USoundFactory>())
+					// Setup sane defaults for importing localized sound waves
+					SoundWaveFactory->bAutoCreateCue = false;
+					SoundWaveFactory->SuppressImportDialogs();
+
+					UPackage* TempPackage = CreatePackage(TEXT("/Temp/WaveTable"));
+					if (USoundWave* SoundWave = ImportObject<USoundWave>(TempPackage, "ImportedWaveTable", RF_Public | RF_Standalone, *InFilePath, nullptr, SoundWaveFactory))
 					{
-						// Setup sane defaults for importing localized sound waves
-						SoundWaveFactory->bAutoCreateCue = false;
-						SoundWaveFactory->SuppressImportDialogs();
+						TArray<uint8> RawPCMData;
+						uint16 NumChannels = 0;
+						uint32 FileSampleRate = 0;
+						SoundWave->GetImportedSoundWaveData(RawPCMData, FileSampleRate, NumChannels);
+						OutSampleRate = static_cast<int32>(FileSampleRate);
 
-						UPackage* TempPackage = CreatePackage(TEXT("/Temp/WaveTable"));
-						if (USoundWave* SoundWave = ImportObject<USoundWave>(TempPackage, "ImportedWaveTable", RF_Public | RF_Standalone, *InFilePath, nullptr, SoundWaveFactory))
+						// Deinterleave & perform BDC if necessary
+						if (NumChannels > 0)
 						{
-							TArray<uint8> RawPCMData;
-							uint32 SampleRate = 0;
-							uint16 NumChannels = 0;
-							SoundWave->GetImportedSoundWaveData(RawPCMData, SampleRate, NumChannels);
+							const int32 NumFrames = RawPCMData.Num() / NumChannels / sizeof(int16);
+							const int16* RawDataPtr = (const int16*)(RawPCMData.GetData());
+							const int32 ChannelOffset = (InChannelIndex % NumChannels);
+
+							OutData.Zero(NumFrames);
 
-							if (NumChannels > 0)
+							constexpr float ConversionValue = 1.0f / static_cast<float>(TNumericLimits<int16>::Max());
+							switch (OutData.GetBitDepth())
 							{
-								const int32 NumFrames = RawPCMData.Num() / NumChannels / sizeof(int16);
-								OutPCMData.Empty();
-								OutPCMData.AddZeroed(NumFrames);
-								const int16* RawDataPtr = (const int16*)(RawPCMData.GetData());
-								const int32 ChannelOffset = (InChannelIndex % NumChannels);
-								constexpr float Max16BitAsFloat = static_cast<float>(TNumericLimits<int16>::Max());
-								for (int32 i = 0; i < OutPCMData.Num(); ++i)
+								case EWaveTableBitDepth::IEEE_Float:
+								{
+									TArrayView<float> DataView;
+									OutData.GetDataView(DataView);
+									for (int32 Index = 0; Index < DataView.Num(); ++Index)
+									{
+										DataView[Index] = (RawDataPtr[(Index * NumChannels) + ChannelOffset]) * ConversionValue;
+									}
+
+									if (!DataView.IsEmpty())
+									{
+										OutData.SetFinalValue(DataView.Last());
+									}
+								}
+								break;
+
+								case EWaveTableBitDepth::PCM_16:
 								{
-									OutPCMData[i] = (RawDataPtr[(i * NumChannels) + ChannelOffset]) / Max16BitAsFloat;
+									TArrayView<int16> DataView;
+									OutData.GetDataView(DataView);
+									for (int32 Index = 0; Index < DataView.Num(); ++Index)
+									{
+										DataView[Index] = (RawDataPtr[(Index * NumChannels) + ChannelOffset]);
+									}
+
+									if (!DataView.IsEmpty())
+									{
+										OutData.SetFinalValue(DataView.Last() * ConversionValue);
+									}
+								}
+								break;
+
+								default:
+								{
+									static_assert(static_cast<int32>(EWaveTableBitDepth::COUNT) == 2, "Possible missing switch case coverage for 'EWaveTableBitDepth'");
+									checkNoEntry();
 								}
+								break;
+							}
+						}
+						else
+						{
+							UE_LOG(LogWaveTableEditor, Error, TEXT("Failed to import/reimport file '%s': Invalid num channels '%u' imported"), *InFilePath, NumChannels);
+						}
+					}
+					else
+					{
+						UE_LOG(LogWaveTableEditor, Error, TEXT("Failed to import/reimport file '%s'"), *InFilePath);
+					}
+				}
+			}
+		}
+
+		void LoadPCMChannel(const FString& InFilePath, int32 InChannelIndex, TArray<float>& OutPCMData)
+		{
+			OutPCMData.Empty();
+
+			if (!InFilePath.IsEmpty())
+			{
+				if (USoundFactory* SoundWaveFactory = NewObject<USoundFactory>())
+				{
+					// Setup sane defaults for importing localized sound waves
+					SoundWaveFactory->bAutoCreateCue = false;
+					SoundWaveFactory->SuppressImportDialogs();
+
+					UPackage* TempPackage = CreatePackage(TEXT("/Temp/WaveTable"));
+					if (USoundWave* SoundWave = ImportObject<USoundWave>(TempPackage, "ImportedWaveTable", RF_Public | RF_Standalone, *InFilePath, nullptr, SoundWaveFactory))
+					{
+						TArray<uint8> RawPCMData;
+						uint32 SampleRate = 0;
+						uint16 NumChannels = 0;
+						SoundWave->GetImportedSoundWaveData(RawPCMData, SampleRate, NumChannels);
+
+						if (NumChannels > 0)
+						{
+							const int32 NumFrames = RawPCMData.Num() / NumChannels / sizeof(int16);
+							OutPCMData.Empty();
+							OutPCMData.AddZeroed(NumFrames);
+							const int16* RawDataPtr = (const int16*)(RawPCMData.GetData());
+							const int32 ChannelOffset = (InChannelIndex % NumChannels);
+							constexpr float Max16BitAsFloat = static_cast<float>(TNumericLimits<int16>::Max());
+							for (int32 i = 0; i < OutPCMData.Num(); ++i)
+							{
+								OutPCMData[i] = (RawDataPtr[(i * NumChannels) + ChannelOffset]) / Max16BitAsFloat;
 							}
 						}
 					}
 				}
 			}
-		} // namespace FileUtilities
-	} // namespace Editor
-} // namespace WaveTable
+		}
+	} // namespace FileUtilities
+} // namespace WaveTable::Editor
diff --git a/Engine/Plugins/Runtime/WaveTable/Source/WaveTableEditor/Private/WaveTableTransformLayout.cpp b/Engine/Plugins/Runtime/WaveTable/Source/WaveTableEditor/Private/WaveTableTransformLayout.cpp
--- a/Engine/Plugins/Runtime/WaveTable/Source/WaveTableEditor/Private/WaveTableTransformLayout.cpp
+++ b/Engine/Plugins/Runtime/WaveTable/Source/WaveTableEditor/Private/WaveTableTransformLayout.cpp
@@ -5,292 +5,330 @@
 #include "IDetailChildrenBuilder.h"
 #include "PropertyCustomizationHelpers.h"
 #include "Subsystems/AssetEditorSubsystem.h"
+#include "WaveTable.h"
 #include "WaveTableBank.h"
 #include "WaveTableFileUtilities.h"
+#include "WaveTableSettings.h"
 #include "Widgets/Layout/SBox.h"
 #include "Widgets/SToolTip.h"
 
 #define LOCTEXT_NAMESPACE "WaveTableEditor"
 
 
-namespace WaveTable
+namespace WaveTable::Editor
 {
-	namespace Editor
+	void FWaveTableDataLayoutCustomization::CustomizeHeader(TSharedRef<IPropertyHandle> StructPropertyHandle, FDetailWidgetRow& HeaderRow, IPropertyTypeCustomizationUtils& CustomizationUtils)
 	{
-		void FTransformLayoutCustomizationBase::CustomizeHeader(TSharedRef<IPropertyHandle> StructPropertyHandle, FDetailWidgetRow& HeaderRow, IPropertyTypeCustomizationUtils& CustomizationUtils)
+	}
+
+	void FWaveTableDataLayoutCustomization::CustomizeChildren(TSharedRef<IPropertyHandle> StructPropertyHandle, IDetailChildrenBuilder& ChildBuilder, IPropertyTypeCustomizationUtils& StructCustomizationUtils)
+	{
+		uint32 NumChildren;
+		StructPropertyHandle->GetNumChildren(NumChildren);
+
+		TMap<FName, TSharedPtr<IPropertyHandle>> PropertyHandles;
+
+		for (uint32 ChildIndex = 0; ChildIndex < NumChildren; ++ChildIndex)
 		{
+			TSharedRef<IPropertyHandle> ChildHandle = StructPropertyHandle->GetChildHandle(ChildIndex).ToSharedRef();
+			const FName PropertyName = ChildHandle->GetProperty()->GetFName();
+			PropertyHandles.Add(PropertyName, ChildHandle);
 		}
 
-		void FTransformLayoutCustomizationBase::CustomizeChildren(TSharedRef<IPropertyHandle> StructPropertyHandle, IDetailChildrenBuilder& ChildBuilder, IPropertyTypeCustomizationUtils& StructCustomizationUtils)
-		{
-			uint32 NumChildren;
-			StructPropertyHandle->GetNumChildren(NumChildren);
+		TSharedRef<IPropertyHandle> BitDepthHandle = PropertyHandles.FindChecked(FWaveTableData::GetBitDepthPropertyName()).ToSharedRef();
+		ChildBuilder.AddProperty(BitDepthHandle);
+	}
+
+	void FTransformLayoutCustomizationBase::CustomizeHeader(TSharedRef<IPropertyHandle> StructPropertyHandle, FDetailWidgetRow& HeaderRow, IPropertyTypeCustomizationUtils& CustomizationUtils)
+	{
+	}
 
-			FProperty* Property = nullptr;
-			StructPropertyHandle->GetValue(Property);
+	void FTransformLayoutCustomizationBase::CustomizeChildren(TSharedRef<IPropertyHandle> StructPropertyHandle, IDetailChildrenBuilder& ChildBuilder, IPropertyTypeCustomizationUtils& StructCustomizationUtils)
+	{
+		uint32 NumChildren;
+		StructPropertyHandle->GetNumChildren(NumChildren);
 
-			TMap<FName, TSharedPtr<IPropertyHandle>> PropertyHandles;
+		TMap<FName, TSharedPtr<IPropertyHandle>> PropertyHandles;
 
-			for (uint32 ChildIndex = 0; ChildIndex < NumChildren; ++ChildIndex)
-			{
-				TSharedRef<IPropertyHandle> ChildHandle = StructPropertyHandle->GetChildHandle(ChildIndex).ToSharedRef();
-				const FName PropertyName = ChildHandle->GetProperty()->GetFName();
-				PropertyHandles.Add(PropertyName, ChildHandle);
-			}
+		for (uint32 ChildIndex = 0; ChildIndex < NumChildren; ++ChildIndex)
+		{
+			TSharedRef<IPropertyHandle> ChildHandle = StructPropertyHandle->GetChildHandle(ChildIndex).ToSharedRef();
+			const FName PropertyName = ChildHandle->GetProperty()->GetFName();
+			PropertyHandles.Add(PropertyName, ChildHandle);
+		}
 
-			CurveHandle = PropertyHandles.FindChecked(GET_MEMBER_NAME_CHECKED(FWaveTableTransform, Curve)).ToSharedRef();
-			CustomizeCurveSelector(ChildBuilder);
+		CurveHandle = PropertyHandles.FindChecked(GET_MEMBER_NAME_CHECKED(FWaveTableTransform, Curve)).ToSharedRef();
+		CustomizeCurveSelector(ChildBuilder);
 
-			TSharedRef<IPropertyHandle> ScalarHandle = PropertyHandles.FindChecked(GET_MEMBER_NAME_CHECKED(FWaveTableTransform, Scalar)).ToSharedRef();
+		TSharedRef<IPropertyHandle> ScalarHandle = PropertyHandles.FindChecked(GET_MEMBER_NAME_CHECKED(FWaveTableTransform, Scalar)).ToSharedRef();
 
-			ChildBuilder.AddProperty(ScalarHandle)
-				.EditCondition(TAttribute<bool>::Create([this]() { return IsScaleableCurve(); }), nullptr)
-				.Visibility(TAttribute<EVisibility>::Create([this]() { return IsScaleableCurve() ? EVisibility::Visible : EVisibility::Hidden; }));
+		ChildBuilder.AddProperty(ScalarHandle)
+			.EditCondition(TAttribute<bool>::Create([this]() { return IsScaleableCurve(); }), nullptr)
+			.Visibility(TAttribute<EVisibility>::Create([this]() { return IsScaleableCurve() ? EVisibility::Visible : EVisibility::Hidden; }));
 
-			TSharedRef<IPropertyHandle> SharedCurveHandle = PropertyHandles.FindChecked(GET_MEMBER_NAME_CHECKED(FWaveTableTransform, CurveShared)).ToSharedRef();
-			ChildBuilder.AddProperty(SharedCurveHandle)
-				.EditCondition(TAttribute<bool>::Create([this]() { return GetCurve() == EWaveTableCurve::Shared; }), nullptr)
-				.Visibility(TAttribute<EVisibility>::Create([this]() { return GetCurve() == EWaveTableCurve::Shared ? EVisibility::Visible : EVisibility::Hidden; }));
+		TSharedRef<IPropertyHandle> SharedCurveHandle = PropertyHandles.FindChecked(GET_MEMBER_NAME_CHECKED(FWaveTableTransform, CurveShared)).ToSharedRef();
+		ChildBuilder.AddProperty(SharedCurveHandle)
+			.EditCondition(TAttribute<bool>::Create([this]() { return GetCurve() == EWaveTableCurve::Shared; }), nullptr)
+			.Visibility(TAttribute<EVisibility>::Create([this]() { return GetCurve() == EWaveTableCurve::Shared ? EVisibility::Visible : EVisibility::Hidden; }));
 
-			TAttribute<EVisibility> IsWaveTableVisibilityAttribute = TAttribute<EVisibility>::Create([this]()
+		TSharedRef<IPropertyHandle> DurationHandle = PropertyHandles.FindChecked(FWaveTableTransform::GetDurationPropertyName()).ToSharedRef();
+		ChildBuilder.AddProperty(DurationHandle)
+			.EditCondition(TAttribute<bool>::Create([this]()
 			{
-				return GetCurve() == EWaveTableCurve::File ? EVisibility::Visible : EVisibility::Hidden;
-			});
+				const bool bIsSampleRateMode = GetSamplingMode() == EWaveTableSamplingMode::FixedSampleRate;
+				const bool bIsNotFile = GetCurve() != EWaveTableCurve::File;
+				return bIsSampleRateMode && bIsNotFile;
+			}), nullptr)
+			.Visibility(TAttribute<EVisibility>::Create([this]()
+			{
+				const bool bIsSampleRateMode = GetSamplingMode() == EWaveTableSamplingMode::FixedSampleRate;
+				return bIsSampleRateMode ? EVisibility::Visible : EVisibility::Collapsed;
+			}));
 
-			WaveTableOptionsHandle = PropertyHandles.FindChecked(GET_MEMBER_NAME_CHECKED(FWaveTableTransform, WaveTableSettings)).ToSharedRef();
+		TAttribute<EVisibility> IsWaveTableVisibilityAttribute = TAttribute<EVisibility>::Create([this]()
+		{
+			return GetCurve() == EWaveTableCurve::File ? EVisibility::Visible : EVisibility::Hidden;
+		});
 
-			uint32 NumWaveTableOptions = 0;
-			if (WaveTableOptionsHandle->GetNumChildren(NumWaveTableOptions))
-			{
-				ChannelIndexHandle = WaveTableOptionsHandle->GetChildHandle(GET_MEMBER_NAME_CHECKED(FWaveTableSettings, ChannelIndex));
-				ChannelIndexHandle->SetOnPropertyValueChanged(FSimpleDelegate::CreateLambda([this] { CachePCMFromFile(); }));
+		WaveTableOptionsHandle = PropertyHandles.FindChecked(GET_MEMBER_NAME_CHECKED(FWaveTableTransform, WaveTableSettings));
 
-				FilePathHandle = WaveTableOptionsHandle->GetChildHandle(GET_MEMBER_NAME_CHECKED(FWaveTableSettings, FilePath));
-				FilePathHandle->SetOnChildPropertyValueChanged(FSimpleDelegate::CreateLambda([this] { CachePCMFromFile(); }));
+		uint32 NumWaveTableOptions = 0;
+		if (WaveTableOptionsHandle->GetNumChildren(NumWaveTableOptions))
+		{
+			ChannelIndexHandle = WaveTableOptionsHandle->GetChildHandle(GET_MEMBER_NAME_CHECKED(FWaveTableSettings, ChannelIndex));
+			ChannelIndexHandle->SetOnPropertyValueChanged(FSimpleDelegate::CreateLambda([this] { CachePCMFromFile(); }));
 
-				FilePathHandle = FilePathHandle->GetChildHandle(GET_MEMBER_NAME_CHECKED(FFilePath, FilePath));
-			}
+			FilePathHandle = WaveTableOptionsHandle->GetChildHandle(GET_MEMBER_NAME_CHECKED(FWaveTableSettings, FilePath));
+			FilePathHandle->SetOnChildPropertyValueChanged(FSimpleDelegate::CreateLambda([this] { CachePCMFromFile(); }));
+
+			FilePathHandle = FilePathHandle->GetChildHandle(GET_MEMBER_NAME_CHECKED(FFilePath, FilePath));
 
-			ChildBuilder.AddProperty(WaveTableOptionsHandle->AsShared())
-				.EditCondition(TAttribute<bool>::Create([this]() { return GetCurve() == EWaveTableCurve::File; }), nullptr)
-				.Visibility(IsWaveTableVisibilityAttribute);
+			SourceDataHandle = WaveTableOptionsHandle->GetChildHandle(GET_MEMBER_NAME_CHECKED(FWaveTableSettings, SourceData));
+			SourceDataHandle->SetOnChildPropertyValueChanged(FSimpleDelegate::CreateLambda([this] { CachePCMFromFile(); }));
 		}
 
-		void FTransformLayoutCustomizationBase::CustomizeCurveSelector(IDetailChildrenBuilder& ChildBuilder)
+		ChildBuilder.AddProperty(WaveTableOptionsHandle->AsShared())
+			.EditCondition(TAttribute<bool>::Create([this]() { return GetCurve() == EWaveTableCurve::File; }), nullptr)
+			.Visibility(IsWaveTableVisibilityAttribute);
+	}
+
+	void FTransformLayoutCustomizationBase::CustomizeCurveSelector(IDetailChildrenBuilder& ChildBuilder)
+	{
+		check(CurveHandle.IsValid());
+
+		auto GetAll = [this, SupportedCurves = GetSupportedCurves()](TArray<TSharedPtr<FString>>& OutNames, TArray<TSharedPtr<SToolTip>>& OutTooltips, TArray<bool>&)
 		{
-			check(CurveHandle.IsValid());
+			CurveDisplayStringToNameMap.Reset();
+			UEnum* Enum = StaticEnum<EWaveTableCurve>();
+			check(Enum);
 
-			auto GetAll = [this, SupportedCurves = GetSupportedCurves()](TArray<TSharedPtr<FString>>& OutNames, TArray<TSharedPtr<SToolTip>>& OutTooltips, TArray<bool>&)
+			static const int64 MaxValue = static_cast<int64>(EWaveTableCurve::Count);
+			for (int32 i = 0; i < Enum->NumEnums(); ++i)
 			{
-				CurveDisplayStringToNameMap.Reset();
-				UEnum* Enum = StaticEnum<EWaveTableCurve>();
-				check(Enum);
-
-				static const int64 MaxValue = static_cast<int64>(EWaveTableCurve::Count);
-				for (int32 i = 0; i < Enum->NumEnums(); ++i)
+				EWaveTableCurve Curve = static_cast<EWaveTableCurve>(Enum->GetValueByIndex(i));
+				if (SupportedCurves.Contains(Curve))
 				{
-					EWaveTableCurve Curve = static_cast<EWaveTableCurve>(Enum->GetValueByIndex(i));
-					if (SupportedCurves.Contains(Curve))
-					{
-						TSharedPtr<FString> DisplayString = MakeShared<FString>(Enum->GetDisplayNameTextByIndex(i).ToString());
-
-						const FName Name = Enum->GetNameByIndex(i);
-						CurveDisplayStringToNameMap.Add(*DisplayString.Get(), Name);
-
-						OutTooltips.Emplace(SNew(SToolTip).Text(Enum->GetToolTipTextByIndex(i)));
-						OutNames.Emplace(MoveTemp(DisplayString));
-					}
-				}
-			};
+					TSharedPtr<FString> DisplayString = MakeShared<FString>(Enum->GetDisplayNameTextByIndex(i).ToString());
 
-			auto GetValue = [this]()
-			{
-				UEnum* Enum = StaticEnum<EWaveTableCurve>();
-				check(Enum);
+					const FName Name = Enum->GetNameByIndex(i);
+					CurveDisplayStringToNameMap.Add(*DisplayString.Get(), Name);
 
-				uint8 IntValue;
-				if (CurveHandle->GetValue(IntValue) == FPropertyAccess::Success)
-				{
-					return Enum->GetDisplayNameTextByValue(IntValue).ToString();
+					OutTooltips.Emplace(SNew(SToolTip).Text(Enum->GetToolTipTextByIndex(i)));
+					OutNames.Emplace(MoveTemp(DisplayString));
 				}
+			}
+		};
 
-				return Enum->GetDisplayNameTextByValue(static_cast<int32>(EWaveTableCurve::Count)).ToString();
-			};
+		auto GetValue = [this]()
+		{
+			UEnum* Enum = StaticEnum<EWaveTableCurve>();
+			check(Enum);
 
-			auto SelectedValue = [this](const FString& InSelected)
+			uint8 IntValue;
+			if (CurveHandle->GetValue(IntValue) == FPropertyAccess::Success)
 			{
-				UEnum* Enum = StaticEnum<EWaveTableCurve>();
-				check(Enum);
+				return Enum->GetDisplayNameTextByValue(IntValue).ToString();
+			}
 
-				const FName& Name = CurveDisplayStringToNameMap.FindChecked(InSelected);
-				const uint8 Value = static_cast<uint8>(Enum->GetValueByName(Name, EGetByNameFlags::ErrorIfNotFound));
-				ensure(CurveHandle->SetValue(Value) == FPropertyAccess::Success);
-			};
+			return Enum->GetDisplayNameTextByValue(static_cast<int32>(EWaveTableCurve::Count)).ToString();
+		};
 
-			static const FText CurveDisplayName = LOCTEXT("WaveTableTransformCurveProperty", "Curve");
-			ChildBuilder.AddCustomRow(CurveDisplayName)
-				.NameContent()
-				[
-					CurveHandle->CreatePropertyNameWidget()
-				]
-				.ValueContent()
-				.MinDesiredWidth(150.0f)
-				[
-					SNew(SHorizontalBox)
-					+ SHorizontalBox::Slot()
-					.VAlign(VAlign_Center)
+		auto SelectedValue = [this](const FString& InSelected)
+		{
+			UEnum* Enum = StaticEnum<EWaveTableCurve>();
+			check(Enum);
+
+			const FName& Name = CurveDisplayStringToNameMap.FindChecked(InSelected);
+			const uint8 Value = static_cast<uint8>(Enum->GetValueByName(Name, EGetByNameFlags::ErrorIfNotFound));
+			ensure(CurveHandle->SetValue(Value) == FPropertyAccess::Success);
+		};
+
+		static const FText CurveDisplayName = LOCTEXT("WaveTableTransformCurveProperty", "Curve");
+		ChildBuilder.AddCustomRow(CurveDisplayName)
+			.NameContent()
+			[
+				CurveHandle->CreatePropertyNameWidget()
+			]
+			.ValueContent()
+			.MinDesiredWidth(150.0f)
+			[
+				SNew(SHorizontalBox)
+				+ SHorizontalBox::Slot()
+				.VAlign(VAlign_Center)
+			[
+				SNew(SBox)
 				[
-					SNew(SBox)
-					[
-						PropertyCustomizationHelpers::MakePropertyComboBox(
-							CurveHandle,
-							FOnGetPropertyComboBoxStrings::CreateLambda(GetAll),
-							FOnGetPropertyComboBoxValue::CreateLambda(GetValue),
-							FOnPropertyComboBoxValueSelected::CreateLambda(SelectedValue)
-						)
-					]
+					PropertyCustomizationHelpers::MakePropertyComboBox(
+						CurveHandle,
+						FOnGetPropertyComboBoxStrings::CreateLambda(GetAll),
+						FOnGetPropertyComboBoxValue::CreateLambda(GetValue),
+						FOnPropertyComboBoxValueSelected::CreateLambda(SelectedValue)
+					)
 				]
-			];
-		}
+			]
+		];
+	}
 
-		void FTransformLayoutCustomizationBase::CachePCMFromFile()
+	EWaveTableSamplingMode FTransformLayoutCustomizationBase::GetSamplingMode() const
+	{
+		TArray<UObject*> OuterObjects;
+		WaveTableOptionsHandle->GetOuterObjects(OuterObjects);
+		for (UObject* Object : OuterObjects)
 		{
-			if (FWaveTableTransform* Transform = GetTransform())
+			if (UWaveTableBank* Bank = Cast<UWaveTableBank>(Object))
 			{
-				bool bLoadSucceeded = false;
-				FString FilePath;
-				if (ensure(FilePathHandle->GetValue(FilePath) == FPropertyAccess::Success))
-				{
-					int32 ChannelIndex = 0;
-					if (ensure(ChannelIndexHandle->GetValue(ChannelIndex) == FPropertyAccess::Success))
-					{
-						TArray<float>& SourcePCMData = Transform->WaveTableSettings.SourcePCMData;
-						WaveTable::Editor::FileUtilities::LoadPCMChannel(FilePath, ChannelIndex, SourcePCMData);
-						bLoadSucceeded = true;
-					}
-				}
-
-				if (!bLoadSucceeded)
-				{
-					Transform->WaveTableSettings.SourcePCMData.Empty();
-				}
+				return Bank->SampleMode;
 			}
+		}
+
+		return EWaveTableSamplingMode::FixedResolution;
+	}
 
-			if (GEditor)
+	void FTransformLayoutCustomizationBase::CachePCMFromFile()
+	{
+		using namespace WaveTable::Editor;
+
+		TArray<UObject*> OuterObjects;
+		WaveTableOptionsHandle->GetOuterObjects(OuterObjects);
+		for (UObject* Object : OuterObjects)
+		{
+			if (UWaveTableBank* Bank = Cast<UWaveTableBank>(Object))
 			{
-				if (UAssetEditorSubsystem* AssetEditorSubsystem = GEditor->GetEditorSubsystem<UAssetEditorSubsystem>())
+				if (FWaveTableTransform* Transform = GetTransform())
 				{
-					TArray<UObject*> OuterObjects;
-					WaveTableOptionsHandle->GetOuterObjects(OuterObjects);
-					for (UObject* Object : OuterObjects)
-					{
-						if (UWaveTableBank* Bank = Cast<UWaveTableBank>(Object))
-						{
-							Bank->RefreshWaveTables();
-						}
-					}
+					FWaveTableSettings& Settings = Transform->WaveTableSettings;
+					const FString& FilePath = Settings.FilePath.FilePath;
+					const int32 ChannelIndex = Settings.ChannelIndex;
+					FWaveTableData& SourceData = Settings.SourceData;
+					int32& SampleRate = Settings.SourceSampleRate;
+					FileUtilities::LoadPCMChannel(FilePath, ChannelIndex, SourceData, SampleRate);
 				}
+
+				// TODO: Only refresh table associated with this transform
+				Bank->RefreshWaveTables();
 			}
 		}
+	}
+
+	bool FTransformLayoutCustomizationBase::IsScaleableCurve() const
+	{
+		static const TArray<EWaveTableCurve> ScalarFilters = { EWaveTableCurve::Exp, EWaveTableCurve::Exp_Inverse, EWaveTableCurve::Log };
+		return ScalarFilters.Contains(GetCurve());
+	}
 
-		bool FTransformLayoutCustomizationBase::IsScaleableCurve() const
+	EWaveTableCurve FTransformLayoutCustomizationBase::GetCurve() const
+	{
+		if (!CurveHandle.IsValid())
 		{
-			static const TArray<EWaveTableCurve> ScalarFilters = { EWaveTableCurve::Exp, EWaveTableCurve::Exp_Inverse, EWaveTableCurve::Log };
-			return ScalarFilters.Contains(GetCurve());
+			return EWaveTableCurve::Linear;
 		}
 
-		EWaveTableCurve FTransformLayoutCustomizationBase::GetCurve() const
+		uint8 CurveValue = static_cast<uint8>(EWaveTableCurve::Linear);
+		CurveHandle->GetValue(CurveValue);
+		return static_cast<EWaveTableCurve>(CurveValue);
+	}
+
+	int32 FTransformLayoutCustomizationBase::GetOwningArrayIndex() const
+	{
+		int32 TableIndex = INDEX_NONE;
+		if (ensure(WaveTableOptionsHandle.IsValid()))
 		{
-			if (!CurveHandle.IsValid())
+			TSharedPtr<IPropertyHandleArray> ParentArray;
+			TSharedPtr<IPropertyHandle> ChildHandle = WaveTableOptionsHandle;
+			while (ChildHandle.IsValid() && TableIndex == INDEX_NONE)
 			{
-				return EWaveTableCurve::Linear;
+				TableIndex = ChildHandle->GetIndexInArray();
+				ChildHandle = ChildHandle->GetParentHandle();
 			}
-
-			uint8 CurveValue = static_cast<uint8>(EWaveTableCurve::Linear);
-			CurveHandle->GetValue(CurveValue);
-			return static_cast<EWaveTableCurve>(CurveValue);
 		}
 
-		int32 FTransformLayoutCustomizationBase::GetOwningArrayIndex() const
-		{
-			int32 TableIndex = INDEX_NONE;
-			if (ensure(WaveTableOptionsHandle.IsValid()))
-			{
-				TSharedPtr<IPropertyHandleArray> ParentArray;
-				TSharedPtr<IPropertyHandle> ChildHandle = WaveTableOptionsHandle;
-				while (ChildHandle.IsValid() && TableIndex == INDEX_NONE)
-				{
-					TableIndex = ChildHandle->GetIndexInArray();
-					ChildHandle = ChildHandle->GetParentHandle();
-				}
-			}
+		return TableIndex;
+	}
 
-			return TableIndex;
-		}
+	TSet<EWaveTableCurve> FTransformLayoutCustomizationBase::GetSupportedCurves() const
+	{
+		UEnum* Enum = StaticEnum<EWaveTableCurve>();
+		check(Enum);
 
-		TSet<EWaveTableCurve> FTransformLayoutCustomizationBase::GetSupportedCurves() const
+		TSet<EWaveTableCurve> Curves;
+		static const int64 MaxValue = static_cast<int64>(EWaveTableCurve::Count);
+		for (int32 i = 0; i < Enum->NumEnums(); ++i)
 		{
-			UEnum* Enum = StaticEnum<EWaveTableCurve>();
-			check(Enum);
-
-			TSet<EWaveTableCurve> Curves;
-			static const int64 MaxValue = static_cast<int64>(EWaveTableCurve::Count);
-			for (int32 i = 0; i < Enum->NumEnums(); ++i)
+			if (Enum->GetValueByIndex(i) < MaxValue)
 			{
-				if (Enum->GetValueByIndex(i) < MaxValue)
+				if (!Enum->HasMetaData(TEXT("Hidden"), i))
 				{
-					if (!Enum->HasMetaData(TEXT("Hidden"), i))
-					{
-						Curves.Add(static_cast<EWaveTableCurve>(Enum->GetValueByIndex(i)));
-					}
+					Curves.Add(static_cast<EWaveTableCurve>(Enum->GetValueByIndex(i)));
 				}
 			}
-
-			return Curves;
 		}
 
-		bool FTransformLayoutCustomization::IsBipolar() const
+		return Curves;
+	}
+
+	bool FTransformLayoutCustomization::IsBipolar() const
+	{
+		if (ensure(WaveTableOptionsHandle.IsValid()))
 		{
-			if (ensure(WaveTableOptionsHandle.IsValid()))
+			TArray<UObject*> OuterObjects;
+			WaveTableOptionsHandle->GetOuterObjects(OuterObjects);
+			if (OuterObjects.Num() == 1)
 			{
-				TArray<UObject*> OuterObjects;
-				WaveTableOptionsHandle->GetOuterObjects(OuterObjects);
-				if (OuterObjects.Num() == 1)
+				if (UWaveTableBank* Bank = Cast<UWaveTableBank>(OuterObjects.Last()))
 				{
-					if (UWaveTableBank* Bank = Cast<UWaveTableBank>(OuterObjects.Last()))
-					{
-						return Bank->bBipolar;
-					}
+					return Bank->bBipolar;
 				}
 			}
-
-			return false;
 		}
 
-		FWaveTableTransform* FTransformLayoutCustomization::GetTransform() const
+		return false;
+	}
+
+	FWaveTableTransform* FTransformLayoutCustomization::GetTransform() const
+	{
+		if (ensure(WaveTableOptionsHandle.IsValid()))
 		{
-			if (ensure(WaveTableOptionsHandle.IsValid()))
+			TArray<UObject*> OuterObjects;
+			WaveTableOptionsHandle->GetOuterObjects(OuterObjects);
+			if (OuterObjects.Num() == 1)
 			{
-				TArray<UObject*> OuterObjects;
-				WaveTableOptionsHandle->GetOuterObjects(OuterObjects);
-				if (OuterObjects.Num() == 1)
+				if (UWaveTableBank* Bank = Cast<UWaveTableBank>(OuterObjects.Last()))
 				{
-					if (UWaveTableBank* Bank = Cast<UWaveTableBank>(OuterObjects.Last()))
+					const int32 TableIndex = GetOwningArrayIndex();
+					if (TableIndex != INDEX_NONE)
 					{
-						const int32 TableIndex = GetOwningArrayIndex();
-						if (TableIndex != INDEX_NONE)
+						PRAGMA_DISABLE_DEPRECATION_WARNINGS
+						if (TableIndex < Bank->Entries.Num())
 						{
-							if (TableIndex < Bank->Entries.Num())
-							{
-								return &Bank->Entries[TableIndex].Transform;
-							}
+							return &Bank->Entries[TableIndex].Transform;
 						}
+						PRAGMA_ENABLE_DEPRECATION_WARNINGS
 					}
 				}
 			}
-
-			return nullptr;
 		}
-	} // namespace Editor
-} // namespace WaveTable
+
+		return nullptr;
+	}
+} // namespace WaveTable::Editor
 #undef LOCTEXT_NAMESPACE // WaveTableEditor
diff --git a/Engine/Plugins/Runtime/WaveTable/Source/WaveTableEditor/Public/WaveTableBankEditor.h b/Engine/Plugins/Runtime/WaveTable/Source/WaveTableEditor/Public/WaveTableBankEditor.h
--- a/Engine/Plugins/Runtime/WaveTable/Source/WaveTableEditor/Public/WaveTableBankEditor.h
+++ b/Engine/Plugins/Runtime/WaveTable/Source/WaveTableEditor/Public/WaveTableBankEditor.h
@@ -70,10 +70,16 @@ namespace WaveTable
 
 			virtual bool GetIsPropertyEditorDisabled() const { return false; }
 
-			// Returns resolution of bank being editor. By default, bank does not support WaveTable generation
+			// Returns resolution of bank being edited. By default, bank does not support WaveTable generation
 			// by being set to no resolution.
 			virtual EWaveTableResolution GetBankResolution() const { return EWaveTableResolution::None; }
 
+			// Returns sampling mode of bank being edited.
+			virtual EWaveTableSamplingMode GetBankSamplingMode() const { return EWaveTableSamplingMode::FixedResolution; }
+
+			// Returns sample rate to use if BankSamplingMode is set to 'SampleRate'.
+			virtual int32 GetBankSampleRate() const { return 0; }
+
 			// Returns whether or not bank is bipolar.  By default, returns false (bank is unipolar), functionally
 			// operating as a unipolar envelope editor.
 			virtual bool GetBankIsBipolar() const { return false; }
@@ -139,6 +145,8 @@ namespace WaveTable
 
 		protected:
 			virtual EWaveTableResolution GetBankResolution() const override;
+			virtual EWaveTableSamplingMode GetBankSamplingMode() const override;
+			virtual int32 GetBankSampleRate() const override;
 			virtual bool GetBankIsBipolar() const override;
 			virtual int32 GetNumTransforms() const override;
 			virtual FWaveTableTransform* GetTransform(int32 InIndex) const override;
diff --git a/Engine/Plugins/Runtime/WaveTable/Source/WaveTableEditor/Public/WaveTableCurveEditorViewStacked.h b/Engine/Plugins/Runtime/WaveTable/Source/WaveTableEditor/Public/WaveTableCurveEditorViewStacked.h
--- a/Engine/Plugins/Runtime/WaveTable/Source/WaveTableEditor/Public/WaveTableCurveEditorViewStacked.h
+++ b/Engine/Plugins/Runtime/WaveTable/Source/WaveTableEditor/Public/WaveTableCurveEditorViewStacked.h
@@ -3,14 +3,17 @@
 
 #include "RichCurveEditorModel.h"
 #include "Views/SCurveEditorViewStacked.h"
-
-enum class ECurveEditorViewID : uint64;
+#include "WaveTableSettings.h"
+#include "WaveTableTransformLayout.h"
 
 #include "WaveTableCurveEditorViewStacked.generated.h"
 
+
 // Forward Declarations
 class UCurveFloat;
 struct FWaveTableTransform;
+enum class ECurveEditorViewID : uint64;
+enum class EWaveTableSamplingMode : uint8;
 
 
 UENUM()
@@ -37,19 +40,20 @@ namespace WaveTable
 			const FText& GetAxesDescriptor() const;
 			const UObject* GetParentObject() const;
 			EWaveTableCurveSource GetSource() const;
-			void Refresh(const FWaveTableTransform& InTransform, int32 InCurveIndex, bool bInIsBipolar);
+			void Refresh(const FWaveTableTransform& InTransform, int32 InCurveIndex, bool bInIsBipolar, EWaveTableSamplingMode InSamplingMode = EWaveTableSamplingMode::FixedResolution);
 
 			virtual ECurveEditorViewID GetViewId() const { return WaveTableViewId; }
 			virtual bool IsReadOnly() const override;
 			virtual FLinearColor GetColor() const override;
 			virtual void GetValueRange(double& MinValue, double& MaxValue) const override;
 
+			float GetCurveDuration() const { return CurveDuration; }
 			int32 GetCurveIndex() const { return CurveIndex; }
 			bool GetIsBipolar() const { return bIsBipolar; }
 			float GetFadeInRatio() const { return FadeInRatio; }
 			float GetFadeOutRatio() const { return FadeOutRatio; }
 			int32 GetNumSamples() const { return NumSamples; }
-
+			EWaveTableSamplingMode GetSamplingMode() const { return SamplingMode; }
 		protected:
 			virtual void RefreshCurveDescriptorText(const FWaveTableTransform& InTransform, FText& OutShortDisplayName, FText& OutInputAxisName, FText& OutOutputAxisName);
 			virtual FColor GetCurveColor() const;
@@ -63,11 +67,13 @@ namespace WaveTable
 			int32 NumSamples = 0;
 
 			EWaveTableCurveSource Source = EWaveTableCurveSource::Unset;
+			EWaveTableSamplingMode SamplingMode = EWaveTableSamplingMode::FixedResolution;
 
 			bool bIsBipolar = false;
 
 			float FadeInRatio = 0.0f;
 			float FadeOutRatio = 0.0f;
+			float CurveDuration = 1.0f;
 
 			FText InputAxisName;
 			FText AxesDescriptor;
diff --git a/Engine/Plugins/Runtime/WaveTable/Source/WaveTableEditor/Public/WaveTableFileUtilities.h b/Engine/Plugins/Runtime/WaveTable/Source/WaveTableEditor/Public/WaveTableFileUtilities.h
--- a/Engine/Plugins/Runtime/WaveTable/Source/WaveTableEditor/Public/WaveTableFileUtilities.h
+++ b/Engine/Plugins/Runtime/WaveTable/Source/WaveTableEditor/Public/WaveTableFileUtilities.h
@@ -2,18 +2,19 @@
 #pragma once
 
 #include "Containers/Array.h"
+#include "WaveTable.h"
 
 
-namespace WaveTable
+namespace WaveTable::Editor
 {
-	namespace Editor
+	namespace FileUtilities
 	{
-		namespace FileUtilities
-		{
-			void WAVETABLEEDITOR_API LoadPCMChannel(const FString& InFilePath, int32 InChannelIndex, TArray<float>& OutPCMData);
-		} // namespace FileUtilities
-	} // namespace Editor
-} // namespace WaveTable
+		void WAVETABLEEDITOR_API LoadPCMChannel(const FString& InFilePath, int32 InChannelIndex, FWaveTableData& OutData, int32& OutSampleRate);
+
+		UE_DEPRECATED(5.3, "Use version of 'LoadPCMChannel' that takes in WaveTableData to support sample rate & bit depth conversion")
+		void WAVETABLEEDITOR_API LoadPCMChannel(const FString& InFilePath, int32 InChannelIndex, TArray<float>& OutPCMData);
+	} // namespace FileUtilities
+} // namespace WaveTable::Editor
 
 #if UE_ENABLE_INCLUDE_ORDER_DEPRECATED_IN_5_2
 #include "Containers/UnrealString.h"
diff --git a/Engine/Plugins/Runtime/WaveTable/Source/WaveTableEditor/Public/WaveTableTransformLayout.h b/Engine/Plugins/Runtime/WaveTable/Source/WaveTableEditor/Public/WaveTableTransformLayout.h
--- a/Engine/Plugins/Runtime/WaveTable/Source/WaveTableEditor/Public/WaveTableTransformLayout.h
+++ b/Engine/Plugins/Runtime/WaveTable/Source/WaveTableEditor/Public/WaveTableTransformLayout.h
@@ -3,64 +3,67 @@
 
 #include "IPropertyTypeCustomization.h"
 
+
+// Forward Declarations
 enum class EWaveTableCurve : uint8;
+enum class EWaveTableSamplingMode : uint8;
 struct FWaveTableTransform;
 
 
-namespace WaveTable
+namespace WaveTable::Editor
 {
-	namespace Editor
+	class WAVETABLEEDITOR_API FWaveTableDataLayoutCustomization : public IPropertyTypeCustomization
 	{
-		class WAVETABLEEDITOR_API FTransformLayoutCustomizationBase : public IPropertyTypeCustomization
+	public:
+		static TSharedRef<IPropertyTypeCustomization> MakeInstance()
 		{
-		public:
-			//~ Begin IPropertyTypeCustomization
-			virtual void CustomizeHeader(TSharedRef<IPropertyHandle> StructPropertyHandle, FDetailWidgetRow& HeaderRow, IPropertyTypeCustomizationUtils& StructCustomizationUtils) override;
-			virtual void CustomizeChildren(TSharedRef<IPropertyHandle> StructPropertyHandle, IDetailChildrenBuilder& ChildBuilder, IPropertyTypeCustomizationUtils& StructCustomizationUtils) override;
-			//~ End IPropertyTypeCustomization
+			return MakeShared<FWaveTableDataLayoutCustomization>();
+		}
 
-		protected:
-			virtual TSet<EWaveTableCurve> GetSupportedCurves() const;
-			virtual FWaveTableTransform* GetTransform() const = 0;
-			virtual bool IsBipolar() const = 0;
+		virtual void CustomizeHeader(TSharedRef<IPropertyHandle> StructPropertyHandle, FDetailWidgetRow& HeaderRow, IPropertyTypeCustomizationUtils& CustomizationUtils) override;
+		virtual void CustomizeChildren(TSharedRef<IPropertyHandle> StructPropertyHandle, IDetailChildrenBuilder& ChildBuilder, IPropertyTypeCustomizationUtils& StructCustomizationUtils) override;
+	};
 
-			void CachePCMFromFile();
-			EWaveTableCurve GetCurve() const;
-			int32 GetOwningArrayIndex() const;
-			bool IsScaleableCurve() const;
+	class WAVETABLEEDITOR_API FTransformLayoutCustomizationBase : public IPropertyTypeCustomization
+	{
+	public:
+		//~ Begin IPropertyTypeCustomization
+		virtual void CustomizeHeader(TSharedRef<IPropertyHandle> StructPropertyHandle, FDetailWidgetRow& HeaderRow, IPropertyTypeCustomizationUtils& StructCustomizationUtils) override;
+		virtual void CustomizeChildren(TSharedRef<IPropertyHandle> StructPropertyHandle, IDetailChildrenBuilder& ChildBuilder, IPropertyTypeCustomizationUtils& StructCustomizationUtils) override;
+		//~ End IPropertyTypeCustomization
 
-			TSharedPtr<IPropertyHandle> CurveHandle;
-			TSharedPtr<IPropertyHandle> ChannelIndexHandle;
-			TSharedPtr<IPropertyHandle> FilePathHandle;
-			TSharedPtr<IPropertyHandle> WaveTableOptionsHandle;
+	protected:
+		virtual TSet<EWaveTableCurve> GetSupportedCurves() const;
+		virtual FWaveTableTransform* GetTransform() const = 0;
+		virtual bool IsBipolar() const = 0;
+		virtual EWaveTableSamplingMode GetSamplingMode() const;
 
-		private:
-			void CustomizeCurveSelector(IDetailChildrenBuilder& ChildBuilder);
+		void CachePCMFromFile();
+		EWaveTableCurve GetCurve() const;
+		int32 GetOwningArrayIndex() const;
+		bool IsScaleableCurve() const;
 
-			TMap<FString, FName> CurveDisplayStringToNameMap;
-		};
+		TSharedPtr<IPropertyHandle> CurveHandle;
+		TSharedPtr<IPropertyHandle> ChannelIndexHandle;
+		TSharedPtr<IPropertyHandle> FilePathHandle;
+		TSharedPtr<IPropertyHandle> SourceDataHandle;
+		TSharedPtr<IPropertyHandle> WaveTableOptionsHandle;
 
-		class WAVETABLEEDITOR_API FTransformLayoutCustomization : public FTransformLayoutCustomizationBase
-		{
-		public:
-			static TSharedRef<IPropertyTypeCustomization> MakeInstance()
-			{
-				return MakeShared<FTransformLayoutCustomization>();
-			}
+	private:
+		void CustomizeCurveSelector(IDetailChildrenBuilder& ChildBuilder);
 
-			virtual bool IsBipolar() const override;
-			virtual FWaveTableTransform* GetTransform() const override;
-		};
-	} // namespace Editor
-} // namespace WaveTable
+		TMap<FString, FName> CurveDisplayStringToNameMap;
+	};
+
+	class WAVETABLEEDITOR_API FTransformLayoutCustomization : public FTransformLayoutCustomizationBase
+	{
+	public:
+		static TSharedRef<IPropertyTypeCustomization> MakeInstance()
+		{
+			return MakeShared<FTransformLayoutCustomization>();
+		}
 
-#if UE_ENABLE_INCLUDE_ORDER_DEPRECATED_IN_5_2
-#include "Curves/SimpleCurve.h"
-#include "Layout/Visibility.h"
-#include "Misc/Attribute.h"
-#include "PropertyCustomizationHelpers.h"
-#include "PropertyHandle.h"
-#include "SCurveEditor.h"
-#include "WaveTableSettings.h"
-#include "WaveTableTransform.h"
-#endif
+		virtual bool IsBipolar() const override;
+		virtual FWaveTableTransform* GetTransform() const override;
+	};
+} // namespace WaveTable::Editor
diff --git a/Engine/Plugins/Runtime/WaveTable/Source/WaveTableEditor/WaveTableEditor.Build.cs b/Engine/Plugins/Runtime/WaveTable/Source/WaveTableEditor/WaveTableEditor.Build.cs
--- a/Engine/Plugins/Runtime/WaveTable/Source/WaveTableEditor/WaveTableEditor.Build.cs
+++ b/Engine/Plugins/Runtime/WaveTable/Source/WaveTableEditor/WaveTableEditor.Build.cs
@@ -6,27 +6,30 @@ public class WaveTableEditor : ModuleRules
 {
 	public WaveTableEditor(ReadOnlyTargetRules Target) : base(Target)
 	{
-		PrivateDependencyModuleNames.AddRange(
-			new string[]
-			{
-				"AudioEditor",
-				"AudioExtensions",
-				"Core",
-				"CoreUObject",
-				"CurveEditor",
-				"Engine",
-				"EditorFramework",
-				"EditorWidgets",
-				"GameProjectGeneration",
-				"InputCore",
-				"PropertyEditor",
-				"SequenceRecorder",
-				"Slate",
-				"SlateCore",
-				"ToolWidgets",
-				"UnrealEd",
-				"WaveTable",
-			}
-		);
+		PublicDependencyModuleNames.AddRange(new string[]
+		{
+			"AssetTools",
+			"SlateCore",
+		});
+
+		PrivateDependencyModuleNames.AddRange(new string[]
+		{
+			"AudioEditor",
+			"AudioExtensions",
+			"Core",
+			"CoreUObject",
+			"CurveEditor",
+			"Engine",
+			"EditorFramework",
+			"EditorWidgets",
+			"GameProjectGeneration",
+			"InputCore",
+			"PropertyEditor",
+			"SequenceRecorder",
+			"Slate",
+			"ToolWidgets",
+			"UnrealEd",
+			"WaveTable",
+		});
 	}
 }
