Fix LocText collision
#rnx

diff --git a/Engine/Plugins/Runtime/WaveTable/Source/WaveTableEditor/Private/WaveTableCurveEditorViewStacked.cpp b/Engine/Plugins/Runtime/WaveTable/Source/WaveTableEditor/Private/WaveTableCurveEditorViewStacked.cpp
--- a/Engine/Plugins/Runtime/WaveTable/Source/WaveTableEditor/Private/WaveTableCurveEditorViewStacked.cpp
+++ b/Engine/Plugins/Runtime/WaveTable/Source/WaveTableEditor/Private/WaveTableCurveEditorViewStacked.cpp
@@ -140,13 +140,13 @@ namespace WaveTable
 			{
 				case EWaveTableSamplingMode::FixedResolution:
 				{
-					OutInputAxisName = LOCTEXT("WaveTableCurveModel_DefaultInputAxisName", "Unit Phase (Table Index)");
+					OutInputAxisName = LOCTEXT("WaveTableCurveModel_DefaultFixedResolutionInputAxisName", "Unit Phase (Table Index)");
 				}
 				break;
 
 				case EWaveTableSamplingMode::FixedSampleRate:
 				{
-					OutInputAxisName = LOCTEXT("WaveTableCurveModel_DefaultInputAxisName", "Time");
+					OutInputAxisName = LOCTEXT("WaveTableCurveModel_DefaultFixedSampleRateInputAxisName", "Time");
 				}
 				break;
 
