Fix missing value initialization & invalid clamped value which should be non-zero
#rnx

diff --git a/Engine/Plugins/Runtime/WaveTable/Source/WaveTable/Public/WaveTableTransform.h b/Engine/Plugins/Runtime/WaveTable/Source/WaveTable/Public/WaveTableTransform.h
--- a/Engine/Plugins/Runtime/WaveTable/Source/WaveTable/Public/WaveTableTransform.h
+++ b/Engine/Plugins/Runtime/WaveTable/Source/WaveTable/Public/WaveTableTransform.h
@@ -43,7 +43,7 @@ protected:
 	// Duration of curve or file.  Only valid if parent SampleRate is set and SamplingMode is set to 'FixedSampleRate'
 	// (If set to 'FixedResolution', duration is variable depending on the resolution (number of samples in the table data)
 	// and device's sample rate.
-	UPROPERTY(EditAnywhere, meta = (DisplayName = "Duration (Sec)", Category = Input, ClampMin = "0.0", UIMin = "0.0"))
+	UPROPERTY(EditAnywhere, meta = (DisplayName = "Duration (Sec)", Category = Input, ClampMin = "0.001", UIMin = "0.001"))
 	float Duration = 1.0f;
 
 public:
@@ -118,7 +118,7 @@ public:
 	TArray<float> WaveTable;
 
 	UPROPERTY(meta = (Deprecated = "5.3", DeprecationMessage = "WaveTable data now supports multiple bit depths and FinalValue has moved to TableData"))
-	float FinalValue;
+	float FinalValue = 0.0f;
 #endif // WITH_EDITORONLY_DATA
 
 private:
