Fixes for changes in WaveTableBanks that caused confusing/broken UX around custom/shared WaveTableBank entry curves
#rnx

diff --git a/Engine/Plugins/Runtime/WaveTable/Source/WaveTable/Private/WaveTableBank.cpp b/Engine/Plugins/Runtime/WaveTable/Source/WaveTable/Private/WaveTableBank.cpp
--- a/Engine/Plugins/Runtime/WaveTable/Source/WaveTable/Private/WaveTableBank.cpp
+++ b/Engine/Plugins/Runtime/WaveTable/Source/WaveTable/Private/WaveTableBank.cpp
@@ -223,6 +223,12 @@ void UWaveTableBank::RefreshWaveTables()
 	WaveTableSizeMB = 0.0f;
 	for (FWaveTableBankEntry& Entry : Entries)
 	{
+		if (Entry.Transform.Curve == EWaveTableCurve::Custom && Entry.Transform.CurveCustom.GetNumKeys() < 1)
+		{
+			Entry.Transform.CurveCustom.AddKey(0.0f, 0.0f);
+			Entry.Transform.CurveCustom.AddKey(1.0f, 1.0f);
+		}
+
 		FWaveTableData& RuntimeData = Entry.Transform.TableData;
 		const int32 NumSamples = ::WaveTable::BankPrivate::GetWaveTableNumSamples(*this, Entry);
 
diff --git a/Engine/Plugins/Runtime/WaveTable/Source/WaveTableEditor/Private/WaveTableCurveEditorViewStacked.cpp b/Engine/Plugins/Runtime/WaveTable/Source/WaveTableEditor/Private/WaveTableCurveEditorViewStacked.cpp
--- a/Engine/Plugins/Runtime/WaveTable/Source/WaveTableEditor/Private/WaveTableCurveEditorViewStacked.cpp
+++ b/Engine/Plugins/Runtime/WaveTable/Source/WaveTableEditor/Private/WaveTableCurveEditorViewStacked.cpp
@@ -93,17 +93,40 @@ namespace WaveTable
 				: Color.Desaturate(0.35f);
 		}
 
+		void FWaveTableCurveModel::GetTimeRange(double& MinValue, double& MaxValue) const
+		{
+			if (Source == EWaveTableCurveSource::Shared)
+			{
+				MinValue = 0.0;
+				MaxValue = 1.0;
+			}
+			else
+			{
+				FRichCurveEditorModelRaw::GetTimeRange(MinValue, MaxValue);
+				if (FMath::IsNearlyEqual(MinValue, MaxValue))
+				{
+					MinValue -= 0.01;
+					MaxValue += 0.01;
+				}
+			}
+		}
+
 		void FWaveTableCurveModel::GetValueRange(double& MinValue, double& MaxValue) const
 		{
-			if (FadeInRatio > 0.0f || FadeOutRatio > 0.0f)
+			if (Source == EWaveTableCurveSource::Shared || (FadeInRatio > 0.0f || FadeOutRatio > 0.0f))
 			{
-				MaxValue = 1.0f;
+				MaxValue = 1.0;
 
-				MinValue = bIsBipolar ? -1.0f : 0.0f;
+				MinValue = bIsBipolar ? -1.0 : 0.0;
 			}
 			else
 			{
 				FRichCurveEditorModelRaw::GetValueRange(MinValue, MaxValue);
+				if (FMath::IsNearlyEqual(MinValue, MaxValue))
+				{
+					MinValue -= 0.01;
+					MaxValue += 0.01;
+				}
 			}
 		}
 
@@ -367,54 +390,88 @@ namespace WaveTable
 
 				if (const FWaveTableCurveModel* EditorModel = static_cast<const FWaveTableCurveModel*>(CurveModel))
 				{
-					auto DrawFade = [this, &CurveEditor, &DrawInfo, &OutDrawElements, DrawEffects, BaseLayerId, ModelID = It.Key()](float StartPosX, float EndPosX, bool bFadeIn, bool bIsBipolar)
+					static const FLinearColor FadeColor = FLinearColor::White.CopyWithNewOpacity(0.75f);
+
+					const int32 GridOverlayLayerId = DrawInfo.GetBaseLayerId() + CurveViewConstants::ELayerOffset::GridOverlays;
+					const int32 FadeLayerId = BaseLayerId + CurveViewConstants::ELayerOffset::Curves;
+					const FCurveModelID ModelID = It.Key();
+					const bool bIsBipolar = EditorModel->GetIsBipolar();
+
+					enum ERegionType : uint8 { FadeIn, FadeOut, Invalid };
+					auto DrawRegion = [&](double StartPosX, double EndPosX, ERegionType RegionType, const FSlateBrush* BoxBrush)
 					{
-						const SCurveEditorView* View = CurveEditor->FindFirstInteractiveView(ModelID);
+						const SCurveEditorView* View = CurveEditor->FindFirstInteractiveView(It.Key());
 						FCurveEditorScreenSpace CurveSpace = View->GetCurveSpace(ModelID);
 
 						StartPosX = CurveSpace.SecondsToScreen(StartPosX);
 						EndPosX = CurveSpace.SecondsToScreen(EndPosX);
-						static const FLinearColor FadeColor = FLinearColor::White.CopyWithNewOpacity(0.75f);
-
-						DrawInfo.LinePoints[0].X = StartPosX;
-						DrawInfo.LinePoints[0].Y = CurveSpace.ValueToScreen(0.0f);
-						DrawInfo.LinePoints[1].X = EndPosX;
-						DrawInfo.LinePoints[1].Y = CurveSpace.ValueToScreen(1.0f);
-						const int32 FadeLayerId = BaseLayerId + CurveViewConstants::ELayerOffset::Curves;
 
-						FSlateDrawElement::MakeLines(OutDrawElements, FadeLayerId, DrawInfo.PaintGeometry, DrawInfo.LinePoints, DrawEffects, FadeColor, false);
+						DrawInfo.LinePoints[0].X = (double)StartPosX;
+						DrawInfo.LinePoints[1].X = (double)EndPosX;
 
-						if (bIsBipolar)
+						double Width = 0.0;
+						bool bDrawBox = true;
+						switch (RegionType)
 						{
-							DrawInfo.LinePoints[1].Y = CurveSpace.ValueToScreen(-1.0f);
-							FSlateDrawElement::MakeLines(OutDrawElements, FadeLayerId, DrawInfo.PaintGeometry, DrawInfo.LinePoints, DrawEffects, FadeColor, false);
+							case ERegionType::FadeIn:
+							case ERegionType::FadeOut:
+							{
+								Width = FMath::Abs(DrawInfo.LinePoints[1].X - DrawInfo.LinePoints[0].X);
+								DrawInfo.LinePoints[0].Y = CurveSpace.ValueToScreen(0.0f);
+								DrawInfo.LinePoints[1].Y = CurveSpace.ValueToScreen(1.0f);
+
+								FSlateDrawElement::MakeLines(OutDrawElements, FadeLayerId, DrawInfo.PaintGeometry, DrawInfo.LinePoints, DrawEffects, FadeColor, false);
+
+								if (bIsBipolar)
+								{
+									DrawInfo.LinePoints[1].Y = CurveSpace.ValueToScreen(-1.0f);
+									FSlateDrawElement::MakeLines(OutDrawElements, FadeLayerId, DrawInfo.PaintGeometry, DrawInfo.LinePoints, DrawEffects, FadeColor, false);
+								}
+							}
+							break;
+
+							case ERegionType::Invalid:
+							default:
+							{
+								Width = DrawInfo.LinePoints[1].X - DrawInfo.LinePoints[0].X;
+								bDrawBox = Width > 0.0;
+							}
+							break;
 						}
 
-						const FSlateBrush* WhiteBrush = FAppStyle::GetBrush("WhiteBrush");
-
-						const float Width = FMath::Abs(DrawInfo.LinePoints[1].X - DrawInfo.LinePoints[0].X);
-						const float BoxStart = bFadeIn ? StartPosX : StartPosX - Width;
-						const FPaintGeometry BoxGeometry = DrawInfo.AllottedGeometry->ToPaintGeometry(
-							FVector2D(Width, StackedHeight),
-							FSlateLayoutTransform(FVector2D(BoxStart, DrawInfo.GetPixelTop()))
-						);
+						if (bDrawBox)
+						{
+							const double BoxStart = RegionType == ERegionType::FadeIn || RegionType == ERegionType::Invalid ? StartPosX : StartPosX - Width;
+							const FPaintGeometry BoxGeometry = DrawInfo.AllottedGeometry->ToPaintGeometry(
+								FVector2D(Width, StackedHeight),
+								FSlateLayoutTransform(FVector2D(BoxStart, DrawInfo.GetPixelTop()))
+							);
 
-						const int32 GridOverlayLayerId = DrawInfo.GetBaseLayerId() + CurveViewConstants::ELayerOffset::GridOverlays;
-						FSlateDrawElement::MakeBox(OutDrawElements, GridOverlayLayerId, BoxGeometry, WhiteBrush, DrawEffects, FadeColor.CopyWithNewOpacity(0.15f));
+							FSlateDrawElement::MakeBox(OutDrawElements, GridOverlayLayerId, BoxGeometry, BoxBrush, DrawEffects, FadeColor.CopyWithNewOpacity(0.15f));
+						}
 					};
 
-					const bool bIsBipolar = EditorModel->GetIsBipolar();
 					const float Duration = EditorModel->GetSamplingMode() == EWaveTableSamplingMode::FixedResolution ? 1.0f : EditorModel->GetCurveDuration();
-					float FadeInRatio = EditorModel->GetFadeInRatio();
+					const float FadeInRatio = EditorModel->GetFadeInRatio();
 					if (FadeInRatio > 0.0f)
 					{
-						DrawFade(0.0f, Duration* FadeInRatio, true, bIsBipolar);
+						DrawRegion(0.0, (double)(Duration * FadeInRatio), ERegionType::FadeIn, WhiteBrush);
 					}
 
-					float FadeOutRatio = EditorModel->GetFadeOutRatio();
+					const float FadeOutRatio = EditorModel->GetFadeOutRatio();
 					if (FadeOutRatio > 0.0f)
 					{
-						DrawFade(Duration, Duration * (1.0f - FadeOutRatio), false, bIsBipolar);
+						DrawRegion((double)Duration, (double)(Duration * (1.0f - FadeOutRatio)), ERegionType::FadeOut, WhiteBrush);
+					}
+
+					if (EditorModel->GetSource() == EWaveTableCurveSource::Shared || EditorModel->GetSource() == EWaveTableCurveSource::Custom)
+					{
+						const SCurveEditorView* View = CurveEditor->FindFirstInteractiveView(It.Key());
+						FCurveEditorScreenSpace CurveSpace = View->GetCurveSpace(ModelID);
+
+						const FSlateBrush* BlackBrush = FAppStyle::GetBrush("BlackBrush");
+						DrawRegion(CurveSpace.GetInputMin(), 0.0, ERegionType::Invalid, BlackBrush);
+						DrawRegion(1.0, CurveSpace.GetInputMax(), ERegionType::Invalid, BlackBrush);
 					}
 				}
 			}
@@ -422,8 +479,8 @@ namespace WaveTable
 
 		void SViewStacked::DrawViewGridLineX(FSlateWindowElementList& OutDrawElements, FGridDrawInfo& DrawInfo, ESlateDrawEffect DrawEffects, double OffsetAlpha, bool bIsMajor) const
 		{
-			double ValueMin = 0.0f;
-			double ValueMax = 0.0f;
+			double ValueMin = 0.0;
+			double ValueMax = 0.0;
 			DrawInfo.GetCurveModel()->GetValueRange(ValueMin, ValueMax);
 
 			const double LowerValue = DrawInfo.GetLowerValue();
@@ -520,7 +577,7 @@ namespace WaveTable
 						const float Value = FCString::Atof(*InOutLabel.ToString());
 						if (Value >= 0 && Value <= 1.0f)
 						{
-							NumSamples = FMath::FloorToInt32(NumSamples * Value);
+							NumSamples = FMath::FloorToInt32(FMath::Abs(NumSamples * Value));
 							InOutLabel = FText::Format(LOCTEXT("WaveTableCurveEditor_ResolutionInputLabelFormat", "{0} ({1})"), InOutLabel, FText::AsNumber(NumSamples));
 						}
 					}
@@ -529,7 +586,33 @@ namespace WaveTable
 
 				case EWaveTableSamplingMode::FixedSampleRate:
 				{
-					InOutLabel = FText::Format(LOCTEXT("WaveTableCurveEditor_DurationInputLabelFormat", "{0}s"), InOutLabel);
+					switch (EditorModel.GetSource())
+					{
+						case EWaveTableCurveSource::Custom:
+						case EWaveTableCurveSource::Shared:
+						{
+							const FText NormalizedLabel = InOutLabel;
+							float Value = FCString::Atof(*InOutLabel.ToString()) * EditorModel.GetCurveDuration();
+							// Avoids negative 0 display issue
+							if (FMath::IsNearlyZero(Value))
+							{
+								InOutLabel = LOCTEXT("WaveTableCurveEditor_FixedSampleRate_ZeroFormat", "0s");
+							}
+							else if (Value > 0.0f)
+							{
+								InOutLabel = FText::Format(LOCTEXT("WaveTableCurveEditor_FixedSampleRate_NormalizedAndDurationInputLabelFormat", "{0} ({1}s)")
+									, NormalizedLabel
+									, FText::AsNumber(Value, &InLabelFormat));
+							}
+						}
+						break;
+
+						default:
+						{
+							InOutLabel = FText::Format(LOCTEXT("WaveTableCurveEditor_DurationInputLabelFormat", "{0}s"), InOutLabel);
+						}
+						break;
+					}
 				}
 				break;
 
@@ -541,6 +624,16 @@ namespace WaveTable
 			};
 		}
 
+		FText SViewStacked::FormatToolTipTime(const FCurveModel& CurveModel, double EvaluatedTime) const
+		{
+			FNumberFormattingOptions FormatOptions;
+			FormatOptions.MaximumFractionalDigits = 3;
+			FText Time = FText::AsNumber(EvaluatedTime, &FormatOptions);
+			FormatInputLabel(static_cast<const FWaveTableCurveModel&>(CurveModel), FormatOptions, Time);
+			return FText::Format(LOCTEXT("WaveTable_PointToolTipTime", "Time: {0}"), Time);
+		}
+
+
 		SViewStacked::FGridDrawInfo::FGridDrawInfo(const FGeometry* InAllottedGeometry, const FCurveEditorScreenSpace& InScreenSpace, FLinearColor InGridColor, int32 InBaseLayerId)
 			: AllottedGeometry(InAllottedGeometry)
 			, ScreenSpace(InScreenSpace)
diff --git a/Engine/Plugins/Runtime/WaveTable/Source/WaveTableEditor/Public/WaveTableCurveEditorViewStacked.h b/Engine/Plugins/Runtime/WaveTable/Source/WaveTableEditor/Public/WaveTableCurveEditorViewStacked.h
--- a/Engine/Plugins/Runtime/WaveTable/Source/WaveTableEditor/Public/WaveTableCurveEditorViewStacked.h
+++ b/Engine/Plugins/Runtime/WaveTable/Source/WaveTableEditor/Public/WaveTableCurveEditorViewStacked.h
@@ -45,6 +45,7 @@ namespace WaveTable
 			virtual ECurveEditorViewID GetViewId() const { return WaveTableViewId; }
 			virtual bool IsReadOnly() const override;
 			virtual FLinearColor GetColor() const override;
+			virtual void GetTimeRange(double& MinValue, double& MaxValue) const override;
 			virtual void GetValueRange(double& MinValue, double& MaxValue) const override;
 
 			float GetCurveDuration() const { return CurveDuration; }
@@ -90,7 +91,7 @@ namespace WaveTable
 				const FGeometry& AllottedGeometry,
 				const FSlateRect& MyCullingRect,
 				FSlateWindowElementList& OutDrawElements,
-				int32 BaseLayerId, 
+				int32 BaseLayerId,
 				const FWidgetStyle& InWidgetStyle,
 				bool bParentEnabled) const override;
 
@@ -99,11 +100,12 @@ namespace WaveTable
 
 			virtual void FormatInputLabel(const FWaveTableCurveModel& EditorModel, const FNumberFormattingOptions& InLabelFormat, FText& InOutLabel) const;
 			virtual void FormatOutputLabel(const FWaveTableCurveModel& EditorModel, const FNumberFormattingOptions& InLabelFormat, FText& InOutLabel) const { }
+			virtual FText FormatToolTipTime(const FCurveModel& CurveModel, double EvaluatedTime) const override;
 
 		private:
 			struct FGridDrawInfo
 			{
-				const FGeometry* AllottedGeometry;
+				const FGeometry* AllottedGeometry = nullptr;
 
 				FPaintGeometry PaintGeometry;
 
