- Fix bug with top/tail where tail percent was calculated off the already topped length (instead of the entire source file's length)

diff --git a/Engine/Plugins/Runtime/WaveTable/Source/WaveTable/Private/WaveTableSettings.cpp b/Engine/Plugins/Runtime/WaveTable/Source/WaveTable/Private/WaveTableSettings.cpp
--- a/Engine/Plugins/Runtime/WaveTable/Source/WaveTable/Private/WaveTableSettings.cpp
+++ b/Engine/Plugins/Runtime/WaveTable/Source/WaveTable/Private/WaveTableSettings.cpp
@@ -8,29 +8,25 @@ void FWaveTableSettings::GetEditSourceBounds(int32& OutSourceTopOffset, int32& O
 	OutSourceTopOffset = 0;
 	OutSourceNumSamples = 0;
 
-	int32 SourceTopOffset = 0;
-	int32 SourceNumSamples = SourceData.GetNumSamples();
-	if (SourceNumSamples == 0)
+	const int32 TotalNumSamples = SourceData.GetNumSamples();
+	if (TotalNumSamples > 0)
 	{
-		return;
-	}
+		OutSourceNumSamples = TotalNumSamples;
+		if (Top > 0.0f)
+		{
+			OutSourceTopOffset = WaveTable::RatioToIndex(Top, TotalNumSamples);
+			OutSourceTopOffset = FMath::Min(OutSourceTopOffset, TotalNumSamples);
+			OutSourceNumSamples -= OutSourceTopOffset;
+		}
 
-	if (Top > 0.0f)
-	{
-		SourceTopOffset = WaveTable::RatioToIndex(Top, SourceNumSamples);
-		SourceNumSamples -= SourceTopOffset;
-	}
+		if (Tail > 0.0f)
+		{
+			int32 SourceTailOffset = WaveTable::RatioToIndex(Tail, TotalNumSamples);
+			SourceTailOffset = FMath::Min(SourceTailOffset, TotalNumSamples);
+			OutSourceNumSamples -= SourceTailOffset;
+		}
 
-	if (Tail > 0.0f)
-	{
-		const int32 SourceTailOffset = WaveTable::RatioToIndex(Tail, SourceNumSamples);
-		SourceNumSamples -= SourceTailOffset;
-	}
-
-	if (SourceTopOffset < SourceData.GetNumSamples())
-	{
-		OutSourceTopOffset = SourceTopOffset;
-		OutSourceNumSamples = SourceNumSamples;
+		OutSourceNumSamples = FMath::Max(OutSourceNumSamples, 0);
 	}
 }
 
