- Fix for calculation of final WaveTable value due to float precision loss through erroneous int32 conversion & potential missed FinalValue copy on resave
#rnx
#tests local using sample MetaSound and printlog

diff --git a/Engine/Plugins/Runtime/WaveTable/Source/WaveTable/Private/WaveTableSampler.cpp b/Engine/Plugins/Runtime/WaveTable/Source/WaveTable/Private/WaveTableSampler.cpp
--- a/Engine/Plugins/Runtime/WaveTable/Source/WaveTable/Private/WaveTableSampler.cpp
+++ b/Engine/Plugins/Runtime/WaveTable/Source/WaveTable/Private/WaveTableSampler.cpp
@@ -480,7 +480,7 @@ namespace WaveTable
 		}
 	}
 
-	int32 FWaveTableSampler::ComputeIndexFinished(TArrayView<const float> InSyncTriggers, TArrayView<float> OutIndicesView)
+	float FWaveTableSampler::ComputeIndexFinished(TArrayView<const float> InSyncTriggers, TArrayView<float> OutIndicesView)
 	{
 		check(!OutIndicesView.IsEmpty());
 
diff --git a/Engine/Plugins/Runtime/WaveTable/Source/WaveTable/Private/WaveTableTransform.cpp b/Engine/Plugins/Runtime/WaveTable/Source/WaveTable/Private/WaveTableTransform.cpp
--- a/Engine/Plugins/Runtime/WaveTable/Source/WaveTable/Private/WaveTableTransform.cpp
+++ b/Engine/Plugins/Runtime/WaveTable/Source/WaveTable/Private/WaveTableTransform.cpp
@@ -445,6 +445,7 @@ void FWaveTableTransform::CreateWaveTable(FWaveTableData& InOutTableData, bool b
 		default:
 		{
 			CopyToWaveTable(InOutTableData, bInBipolar);
+			InOutTableData.SetFinalValue(TableData.GetFinalValue());
 		}
 		break;
 	}
diff --git a/Engine/Plugins/Runtime/WaveTable/Source/WaveTable/Public/WaveTableSampler.h b/Engine/Plugins/Runtime/WaveTable/Source/WaveTable/Public/WaveTableSampler.h
--- a/Engine/Plugins/Runtime/WaveTable/Source/WaveTable/Public/WaveTableSampler.h
+++ b/Engine/Plugins/Runtime/WaveTable/Source/WaveTable/Public/WaveTableSampler.h
@@ -103,7 +103,7 @@ namespace WaveTable
 
 		// If set to one-shot, sets IndexFinished to valid index and resets remaining indices provided past the initial stop to the first index.
 		// Returns last positively increasing index if finished. If not finished, returns the last index value stored in the view.
-		int32 ComputeIndexFinished(TArrayView<const float> InSyncTriggers, TArrayView<float> OutIndicesView);
+		float ComputeIndexFinished(TArrayView<const float> InSyncTriggers, TArrayView<float> OutIndicesView);
 
 		// Utility that takes in necessary data for finalizing a single sample index processing from either a TableData or a TableView.  Covers edge cases for final value.
 		float FinalizeSingleSample(float Index, int32 NumSamples, TArrayView<float> OutSample, float LastTableValue, float FinalValue, FWaveTableSampler::ESingleSampleMode InMode);
