Unclamp Modulation Patches

#Synced-CL 50409854

diff --git a/Engine/Plugins/Runtime/AudioModulation/Source/AudioModulation/Private/SoundControlBusMix.cpp b/Engine/Plugins/Runtime/AudioModulation/Source/AudioModulation/Private/SoundControlBusMix.cpp
--- a/Engine/Plugins/Runtime/AudioModulation/Source/AudioModulation/Private/SoundControlBusMix.cpp
+++ b/Engine/Plugins/Runtime/AudioModulation/Source/AudioModulation/Private/SoundControlBusMix.cpp
@@ -137,7 +137,7 @@ void USoundControlBusMix::OnPropertyChanged(FProperty* Property, EPropertyChange
 						{
 							// Clamp the normalized values--either to the normalized [0,1] range or to the parameter-specified "unclamped" range.
 							// The normalized value is not clamped in the ConvertNormalizedToUnit function below, so we do it here.
-							if (!Parameter->IsModulationParameterClameped())
+							if (!Parameter->IsModulationParameterClamped())
 							{
 								const TRange<float> Range = Parameter->GetNormalizedLimits();
 								Stage.Value.TargetValue = FMath::Clamp(Stage.Value.TargetValue, Range.GetLowerBoundValue(), Range.GetUpperBoundValue());
diff --git a/Engine/Plugins/Runtime/AudioModulation/Source/AudioModulation/Private/SoundModulationParameter.cpp b/Engine/Plugins/Runtime/AudioModulation/Source/AudioModulation/Private/SoundModulationParameter.cpp
--- a/Engine/Plugins/Runtime/AudioModulation/Source/AudioModulation/Private/SoundModulationParameter.cpp
+++ b/Engine/Plugins/Runtime/AudioModulation/Source/AudioModulation/Private/SoundModulationParameter.cpp
@@ -34,7 +34,7 @@ void USoundModulationParameter::PostEditChangeProperty(struct FPropertyChangedEv
 void USoundModulationParameter::RefreshNormalizedValue()
 {
 	float NewNormalizedValue = ConvertUnitToNormalized(Settings.ValueUnit);
-	if (Settings.bClampToNormalizedRange)
+	if (IsModulationParameterClamped())
 	{
 		NewNormalizedValue = FMath::Clamp(NewNormalizedValue, 0.0f, 1.0f);
 	}
@@ -47,7 +47,7 @@ void USoundModulationParameter::RefreshNormalizedValue()
 void USoundModulationParameter::RefreshUnitValue()
 {
 	float NewUnitValue = ConvertNormalizedToUnit(Settings.ValueNormalized);
-	if (Settings.bClampToNormalizedRange)
+	if (IsModulationParameterClamped())
 	{
 		NewUnitValue = FMath::Clamp(NewUnitValue, GetUnitMin(), GetUnitMax());
 	}
@@ -69,7 +69,7 @@ Audio::FModulationParameter USoundModulationParameter::CreateParameter() const
 	Parameter.DefaultValue = GetUnitDefault();
 	Parameter.MinValue = GetUnitMin();
 	Parameter.MaxValue = GetUnitMax();
-	Parameter.bClampToNormalizedRange = Settings.bClampToNormalizedRange;
+	Parameter.bClampToNormalizedRange = IsModulationParameterClamped();
 	Parameter.TypeHash = GetTypeHash(this);
 
 #if WITH_EDITORONLY_DATA
@@ -80,7 +80,7 @@ Audio::FModulationParameter USoundModulationParameter::CreateParameter() const
 	return Parameter;
 }
 
-bool USoundModulationParameter::IsModulationParameterClameped() const
+bool USoundModulationParameter::IsModulationParameterClamped() const
 {
 	bool bClampToNormalizedRange = Settings.bClampToNormalizedRange;
 	Audio::GetModulationParameterClampValue(GetFName(), bClampToNormalizedRange);
@@ -272,7 +272,7 @@ bool USoundModulationParameterVolume::RequiresUnitConversion() const
 
 Audio::FModulationUnitConversionFunction USoundModulationParameterVolume::GetUnitConversionFunction() const
 {
-	return [InUnitMin = GetUnitMin(), InUnitMax = VolumeLimit, bClampToNormalized = IsModulationParameterClameped()](float& InOutValue)
+	return [InUnitMin = GetUnitMin(), InUnitMax = VolumeLimit, bClampToNormalized = IsModulationParameterClamped()](float& InOutValue)
 	{
 		InOutValue = InOutValue > 0.0f
 			? Audio::ConvertToDecibels(InOutValue)
@@ -291,7 +291,7 @@ Audio::FModulationUnitConversionFunction USoundModulationParameterVolume::GetUni
 
 Audio::FModulationNormalizedConversionFunction USoundModulationParameterVolume::GetNormalizedConversionFunction() const
 {
-	return [InUnitMin = GetUnitMin(), InUnitMax = VolumeLimit, bClampToNormalized = IsModulationParameterClameped()](float& InOutValue)
+	return [InUnitMin = GetUnitMin(), InUnitMax = VolumeLimit, bClampToNormalized = IsModulationParameterClamped()](float& InOutValue)
 	{
 		InOutValue = InOutValue < InUnitMin || FMath::IsNearlyEqual(InOutValue, InUnitMin)
 			? 0.0f
diff --git a/Engine/Plugins/Runtime/AudioModulation/Source/AudioModulation/Private/SoundModulationPatchProxy.cpp b/Engine/Plugins/Runtime/AudioModulation/Source/AudioModulation/Private/SoundModulationPatchProxy.cpp
--- a/Engine/Plugins/Runtime/AudioModulation/Source/AudioModulation/Private/SoundModulationPatchProxy.cpp
+++ b/Engine/Plugins/Runtime/AudioModulation/Source/AudioModulation/Private/SoundModulationPatchProxy.cpp
@@ -54,6 +54,8 @@ namespace AudioModulation
 			InSettings.OutputParameter.NormalizedFunction(DefaultValue);
 		}
 
+		bClampOutputToNormalizedRange = InSettings.OutputParameter.bClampToNormalizedRange;
+		
 		// Cache existing proxies to avoid releasing bus state (and potentially referenced bus state) when reinitializing
 		const TArray<FModulationInputProxy> CachedProxies = InputProxies;
 
@@ -63,7 +65,7 @@ namespace AudioModulation
 			InputProxies.Emplace(MoveTemp(Input), OutModSystem);
 		}
 
-		OutputProxy = FModulationOutputProxy(InSettings.OutputParameter.DefaultValue, MoveTemp(InSettings.OutputParameter.MixFunction));
+		OutputProxy = FModulationOutputProxy(DefaultValue, MoveTemp(InSettings.OutputParameter.MixFunction));
 	}
 
 	bool FModulationPatchProxy::IsBypassed() const
@@ -123,6 +125,10 @@ namespace AudioModulation
 		}
 
 		OutputProxy.MixFunction(Value, OutSampleHold);
+		if (bClampOutputToNormalizedRange)
+		{
+			Value = FMath::Clamp(Value, 0.0f, 1.0f);
+		}
 	}
 
 	FModulationPatchRefProxy::FModulationPatchRefProxy()
diff --git a/Engine/Plugins/Runtime/AudioModulation/Source/AudioModulation/Private/SoundModulationPatchProxy.h b/Engine/Plugins/Runtime/AudioModulation/Source/AudioModulation/Private/SoundModulationPatchProxy.h
--- a/Engine/Plugins/Runtime/AudioModulation/Source/AudioModulation/Private/SoundModulationPatchProxy.h
+++ b/Engine/Plugins/Runtime/AudioModulation/Source/AudioModulation/Private/SoundModulationPatchProxy.h
@@ -156,6 +156,9 @@ namespace AudioModulation
 		/** Bypasses the patch and doesn't update modulation value */
 		bool bBypass = true;
 
+		/** Whether or not the final output should be clamped to [0,1], based on the Parameter passed in from FModulationPatchSettings */
+		bool bClampOutputToNormalizedRange = true;
+
 		friend class FAudioModulationSystem;
 	};
 
diff --git a/Engine/Plugins/Runtime/AudioModulation/Source/AudioModulation/Public/SoundModulationParameter.h b/Engine/Plugins/Runtime/AudioModulation/Source/AudioModulation/Public/SoundModulationParameter.h
--- a/Engine/Plugins/Runtime/AudioModulation/Source/AudioModulation/Public/SoundModulationParameter.h
+++ b/Engine/Plugins/Runtime/AudioModulation/Source/AudioModulation/Public/SoundModulationParameter.h
@@ -131,7 +131,7 @@ public:
 
 	UE_API Audio::FModulationParameter CreateParameter() const;
 
-	UE_API bool IsModulationParameterClameped() const;
+	UE_API bool IsModulationParameterClamped() const;
 };
 
 // Linearly scaled value between unit minimum and maximum.
diff --git a/Engine/Plugins/Runtime/AudioModulation/Source/AudioModulation/Public/SoundModulationPatch.h b/Engine/Plugins/Runtime/AudioModulation/Source/AudioModulation/Public/SoundModulationPatch.h
--- a/Engine/Plugins/Runtime/AudioModulation/Source/AudioModulation/Public/SoundModulationPatch.h
+++ b/Engine/Plugins/Runtime/AudioModulation/Source/AudioModulation/Public/SoundModulationPatch.h
@@ -21,6 +21,9 @@ USTRUCT(BlueprintType)
 struct FSoundModulationTransform : public FWaveTableTransform
 {
 	GENERATED_USTRUCT_BODY()
+
+	virtual ~FSoundModulationTransform() override = default;
+	virtual bool ShouldClampToNormalizedRange() const override { return false; }
 };
 PRAGMA_ENABLE_DEPRECATION_WARNINGS
 
diff --git a/Engine/Plugins/Runtime/AudioModulation/Source/AudioModulationEditor/Private/Editors/ModulationPatchCurveEditorViewStacked.cpp b/Engine/Plugins/Runtime/AudioModulation/Source/AudioModulationEditor/Private/Editors/ModulationPatchCurveEditorViewStacked.cpp
--- a/Engine/Plugins/Runtime/AudioModulation/Source/AudioModulationEditor/Private/Editors/ModulationPatchCurveEditorViewStacked.cpp
+++ b/Engine/Plugins/Runtime/AudioModulation/Source/AudioModulationEditor/Private/Editors/ModulationPatchCurveEditorViewStacked.cpp
@@ -25,7 +25,16 @@ namespace PatchCurveViewUtils
 		const float NormalizedValue = FCString::Atof(*InOutLabel.ToString());
 		const float UnitValue = InParameter.ConvertNormalizedToUnit(NormalizedValue);
 		FText UnitLabel = FText::AsNumber(UnitValue, &InNumFormatOptions);
-		InOutLabel = FText::Format(LOCTEXT("ModulationPatchCurveView_UnitFormat", "{0} ({1})"), UnitLabel, InOutLabel);
+
+		// Label for the default parameter value
+		if (InParameter.Settings.ValueNormalized == NormalizedValue)
+		{
+			InOutLabel = FText::Format(LOCTEXT("ModulationPatchCurveView_UnitFormat_Default", "{0} ({1}) (Default)"), UnitLabel, InOutLabel);
+		}
+		else
+		{
+			InOutLabel = FText::Format(LOCTEXT("ModulationPatchCurveView_UnitFormat", "{0} ({1})"), UnitLabel, InOutLabel);
+		}
 	}
 } // namespace PatchCurveViewUtils
 
@@ -86,6 +95,45 @@ FModPatchCurveEditorModel::FModPatchCurveEditorModel(FRichCurve& InRichCurve, UO
 {
 }
 
+void FModPatchCurveEditorModel::GetValueRange(double& MinValue, double& MaxValue) const
+{
+	if (ParentObject.IsValid())
+	{
+		if (USoundModulationPatch* Patch = CastChecked<USoundModulationPatch>(ParentObject))
+		{
+			if (TObjectPtr<USoundModulationParameter> Parameter = Patch->PatchSettings.OutputParameter)
+			{
+				if (!Parameter->IsModulationParameterClamped())
+				{
+					TRange<float> Range = Parameter->GetNormalizedLimits();
+					MinValue = Range.GetLowerBoundValue();
+					MaxValue = Range.GetUpperBoundValue();
+					return;
+				}
+			}
+		}
+	}
+	MinValue = 0.0f;
+	MaxValue = 1.0f;
+}
+
+bool FModPatchCurveEditorModel::ShowDefaultDashedLine() const
+{
+	return true;
+}
+
+double FModPatchCurveEditorModel::GetDefaultDashedLineValue() const
+{
+	if (const USoundModulationPatch* Patch = GetPatch())
+	{
+		if (Patch->PatchSettings.OutputParameter)
+		{
+			return Patch->PatchSettings.OutputParameter->Settings.ValueNormalized;
+		}
+	}
+	return 1.0;
+}
+
 bool FModPatchCurveEditorModel::GetIsBypassed() const
 {
 	if (const USoundModulationPatch* Patch = GetPatch())
diff --git a/Engine/Plugins/Runtime/AudioModulation/Source/AudioModulationEditor/Private/Editors/ModulationPatchCurveEditorViewStacked.h b/Engine/Plugins/Runtime/AudioModulation/Source/AudioModulationEditor/Private/Editors/ModulationPatchCurveEditorViewStacked.h
--- a/Engine/Plugins/Runtime/AudioModulation/Source/AudioModulationEditor/Private/Editors/ModulationPatchCurveEditorViewStacked.h
+++ b/Engine/Plugins/Runtime/AudioModulation/Source/AudioModulationEditor/Private/Editors/ModulationPatchCurveEditorViewStacked.h
@@ -43,6 +43,11 @@ public:
 	virtual bool GetPropertyEditorDisabled() const override { return GetIsBypassed(); }
 	virtual FText GetPropertyEditorDisabledText() const override { return LOCTEXT("ModulationPatchCurveEditorView_Bypassed", "Bypassed"); }
 
+	virtual void GetValueRange(double& MinValue, double& MaxValue) const override;
+
+	virtual bool ShowDefaultDashedLine() const override;
+	virtual double GetDefaultDashedLineValue() const override;
+	
 	bool GetIsBypassed() const;
 
 	USoundModulationPatch* GetPatch();
diff --git a/Engine/Plugins/Runtime/WaveTable/Source/WaveTable/Private/WaveTableTransform.cpp b/Engine/Plugins/Runtime/WaveTable/Source/WaveTable/Private/WaveTableTransform.cpp
--- a/Engine/Plugins/Runtime/WaveTable/Source/WaveTable/Private/WaveTableTransform.cpp
+++ b/Engine/Plugins/Runtime/WaveTable/Source/WaveTable/Private/WaveTableTransform.cpp
@@ -184,9 +184,11 @@ void FWaveTableTransform::SampleCurveUnipolar(FWaveTableData& InOutTableData) co
 
 	TArrayView<float> TableView;
 	InOutTableData.GetDataView(TableView);
-
-	// Clamp the input
-	Audio::ArrayClampInPlace(TableView, 0.0f, 1.0f);
+	
+	if (ShouldClampToNormalizedRange())
+	{
+		Audio::ArrayClampInPlace(TableView, 0.0f, 1.0f);
+	}
 
 	float* ValueData = TableView.GetData();
 	switch (Curve)
@@ -316,7 +318,10 @@ void FWaveTableTransform::SampleCurveUnipolar(FWaveTableData& InOutTableData) co
 		break;
 	}
 
-	Audio::ArrayClampInPlace(TableView, 0.0f, 1.0f);
+	if (ShouldClampToNormalizedRange())
+	{
+		Audio::ArrayClampInPlace(TableView, 0.0f, 1.0f);
+	}
 }
 
 void FWaveTableTransform::SampleWaveTableData(FWaveTableData& InOutTableData) const
diff --git a/Engine/Plugins/Runtime/WaveTable/Source/WaveTable/Public/WaveTableTransform.h b/Engine/Plugins/Runtime/WaveTable/Source/WaveTable/Public/WaveTableTransform.h
--- a/Engine/Plugins/Runtime/WaveTable/Source/WaveTable/Public/WaveTableTransform.h
+++ b/Engine/Plugins/Runtime/WaveTable/Source/WaveTable/Public/WaveTableTransform.h
@@ -18,6 +18,8 @@ struct FWaveTableTransform
 {
 	GENERATED_USTRUCT_BODY()
 
+	virtual ~FWaveTableTransform() = default;
+	
 	/** The curve to apply when transforming the output. */
 	UPROPERTY(EditAnywhere, Category = Input, BlueprintReadWrite, meta = (DisplayName = "Curve Type"))
 	EWaveTableCurve Curve = EWaveTableCurve::Linear;
@@ -47,6 +49,8 @@ protected:
 	// and device's sample rate.
 	UPROPERTY(EditAnywhere, meta = (DisplayName = "Duration (Sec)", Category = Input, ClampMin = "0.001", UIMin = "0.001"))
 	float Duration = 1.0f;
+	
+	virtual bool ShouldClampToNormalizedRange() const { return true; }
 
 public:
 #if WITH_EDITORONLY_DATA
diff --git a/Engine/Plugins/Runtime/WaveTable/Source/WaveTableEditor/Private/WaveTableCurveEditorViewStacked.cpp b/Engine/Plugins/Runtime/WaveTable/Source/WaveTableEditor/Private/WaveTableCurveEditorViewStacked.cpp
--- a/Engine/Plugins/Runtime/WaveTable/Source/WaveTableEditor/Private/WaveTableCurveEditorViewStacked.cpp
+++ b/Engine/Plugins/Runtime/WaveTable/Source/WaveTableEditor/Private/WaveTableCurveEditorViewStacked.cpp
@@ -362,11 +362,41 @@ namespace WaveTable
 				DrawInfo.LinePoints[0].X = 0.0;
 				DrawInfo.LinePoints[1].X = Width;
 
-				DrawViewGridLineX(OutDrawElements, DrawInfo, DrawEffects, 1.0  /* OffsetAlpha */, true  /* bIsMajor */);
-				DrawViewGridLineX(OutDrawElements, DrawInfo, DrawEffects, 0.75 /* OffsetAlpha */, false /* bIsMajor */);
-				DrawViewGridLineX(OutDrawElements, DrawInfo, DrawEffects, 0.5  /* OffsetAlpha */, true  /* bIsMajor */);
-				DrawViewGridLineX(OutDrawElements, DrawInfo, DrawEffects, 0.25 /* OffsetAlpha */, false /* bIsMajor */);
-				DrawViewGridLineX(OutDrawElements, DrawInfo, DrawEffects, 0.0  /* OffsetAlpha */, true  /* bIsMajor */);
+				double DefaultValueHeight = -1.0;
+				if (const FWaveTableCurveModel* WaveTableCurveModel = static_cast<const FWaveTableCurveModel*>(CurveModel))
+				{
+					if (WaveTableCurveModel->ShowDefaultDashedLine())
+					{
+						double ValueMin = 0.0;
+						double ValueMax = 0.0;
+						WaveTableCurveModel->GetValueRange(ValueMin, ValueMax);
+						double DefaultValue = WaveTableCurveModel->GetDefaultDashedLineValue(); 
+						DefaultValueHeight = FMath::GetMappedRangeValueClamped(TRange<double>(ValueMin, ValueMax), TRange<double>(0.0, 1.0), DefaultValue);
+					}
+				}
+
+				// Dashed line on the default value
+				if (const FWaveTableCurveModel* WaveTableCurveModel = static_cast<const FWaveTableCurveModel*>(CurveModel))
+				{
+					if (WaveTableCurveModel->ShowDefaultDashedLine())
+					{	
+						DrawViewGridLineX(OutDrawElements, DrawInfo, DrawEffects, DefaultValueHeight  /* OffsetAlpha */, true  /* bIsMajor */, true /* bDashedLine */);
+					}
+				}
+				
+				TArray<double>OffsetAlphas = { 1.0, 0.75, 0.5, 0.25, 0.0 };
+				for (int OffsetAlphaIndex = 0; OffsetAlphaIndex < OffsetAlphas.Num(); ++OffsetAlphaIndex)
+				{
+					// Skip dashed line
+					double OffsetAlpha = OffsetAlphas[OffsetAlphaIndex];
+					if (FMath::IsNearlyEqual(OffsetAlpha, DefaultValueHeight))
+					{
+						continue;
+					}
+					// Even indices are major lines
+					bool bIsMajor = !static_cast<bool>(OffsetAlphaIndex % 2);
+					DrawViewGridLineX(OutDrawElements, DrawInfo, DrawEffects, OffsetAlphas[OffsetAlphaIndex], bIsMajor, false /* bDashedLine */);
+				}
 
 				// Vertical grid lines
 				{
@@ -479,7 +509,7 @@ namespace WaveTable
 			}
 		}
 
-		void SViewStacked::DrawViewGridLineX(FSlateWindowElementList& OutDrawElements, FGridDrawInfo& DrawInfo, ESlateDrawEffect DrawEffects, double OffsetAlpha, bool bIsMajor) const
+		void SViewStacked::DrawViewGridLineX(FSlateWindowElementList& OutDrawElements, FGridDrawInfo& DrawInfo, ESlateDrawEffect DrawEffects, double OffsetAlpha, bool bIsMajor, bool bDashedLine) const
 		{
 			double ValueMin = 0.0;
 			double ValueMax = 0.0;
@@ -489,16 +519,29 @@ namespace WaveTable
 			const double PixelBottom = DrawInfo.GetPixelBottom();
 			const double PixelTop = DrawInfo.GetPixelTop();
 
-			const TSharedRef<FSlateFontMeasure> FontMeasure = FSlateApplication::Get().GetRenderer()->GetFontMeasureService();
-			const FSlateFontInfo FontInfo = FCoreStyle::Get().GetFontStyle("CurveEd.InfoFont");
-
+			
 			const uint32 LineLayerId = DrawInfo.GetBaseLayerId() + CurveViewConstants::ELayerOffset::GridLines;
 
 			FLinearColor Color = bIsMajor ? DrawInfo.GetMajorGridColor() : DrawInfo.GetMinorGridColor();
 
 			DrawInfo.LinePoints[0].Y = DrawInfo.LinePoints[1].Y = DrawInfo.ScreenSpace.ValueToScreen(LowerValue + OffsetAlpha);
-			FSlateDrawElement::MakeLines(OutDrawElements, LineLayerId, DrawInfo.PaintGeometry, DrawInfo.LinePoints, DrawEffects, Color, false);
 
+			if (bDashedLine)
+			{
+				// There is no FVector2d version of MakeDashedLines, so convert 
+				TArray<FVector2f> FloatPoints;
+				Algo::Transform(DrawInfo.LinePoints, FloatPoints, [](FVector2d DoubleIn)
+					{
+						return FVector2f(DoubleIn);
+					});
+				FSlateDrawElement::MakeDashedLines(OutDrawElements, LineLayerId, DrawInfo.PaintGeometry, MoveTemp(FloatPoints), DrawEffects, Color);
+			}
+			else
+			{
+				FSlateDrawElement::MakeLines(OutDrawElements, LineLayerId, DrawInfo.PaintGeometry, DrawInfo.LinePoints, DrawEffects, Color, false);
+			}
+
+			// Label
 			FText Label = FText::AsNumber(FMath::Lerp(ValueMin, ValueMax, OffsetAlpha), &DrawInfo.LabelFormat);
 
 			if (const FWaveTableCurveModel* CurveModel = static_cast<const FWaveTableCurveModel*>(DrawInfo.GetCurveModel()))
@@ -506,6 +549,9 @@ namespace WaveTable
 				FormatOutputLabel(*CurveModel, DrawInfo.LabelFormat, Label);
 			}
 
+			const TSharedRef<FSlateFontMeasure> FontMeasure = FSlateApplication::Get().GetRenderer()->GetFontMeasureService();
+			const FSlateFontInfo FontInfo = FCoreStyle::Get().GetFontStyle("CurveEd.InfoFont");
+
 			const FVector2D LabelSize = FontMeasure->Measure(Label, FontInfo);
 			double LabelY = FMath::Lerp(PixelBottom, PixelTop, OffsetAlpha);
 	
diff --git a/Engine/Plugins/Runtime/WaveTable/Source/WaveTableEditor/Public/WaveTableCurveEditorViewStacked.h b/Engine/Plugins/Runtime/WaveTable/Source/WaveTableEditor/Public/WaveTableCurveEditorViewStacked.h
--- a/Engine/Plugins/Runtime/WaveTable/Source/WaveTableEditor/Public/WaveTableCurveEditorViewStacked.h
+++ b/Engine/Plugins/Runtime/WaveTable/Source/WaveTableEditor/Public/WaveTableCurveEditorViewStacked.h
@@ -50,6 +50,9 @@ namespace WaveTable
 			UE_API virtual void GetTimeRange(double& MinValue, double& MaxValue) const override;
 			UE_API virtual void GetValueRange(double& MinValue, double& MaxValue) const override;
 
+			UE_API virtual bool ShowDefaultDashedLine() const { return false; }
+			UE_API virtual double GetDefaultDashedLineValue() const { return 1.0; } 
+
 			float GetCurveDuration() const { return CurveDuration; }
 			int32 GetCurveIndex() const { return CurveIndex; }
 			bool GetIsBipolar() const { return bIsBipolar; }
@@ -145,7 +148,7 @@ namespace WaveTable
 				double GetPixelTop() const;
 			};
 
-			UE_API void DrawViewGridLineX(FSlateWindowElementList& OutDrawElements, FGridDrawInfo& DrawInfo, ESlateDrawEffect DrawEffect, double OffsetAlpha, bool bIsMajor) const;
+			UE_API void DrawViewGridLineX(FSlateWindowElementList& OutDrawElements, FGridDrawInfo& DrawInfo, ESlateDrawEffect DrawEffect, double OffsetAlpha, bool bIsMajor, bool bDashedLine = false) const;
 			UE_API void DrawViewGridLineY(const float VerticalLine, FSlateWindowElementList& OutDrawElements, FGridDrawInfo &DrawInfo, ESlateDrawEffect DrawEffects, const FText* Label, bool bIsMajor) const;
 		};
 	} // namespace Editor
